Probable non-unity twitch online subsystem fixes
#rnx

diff --git a/Engine/Plugins/Online/OnlineSubsystemTwitch/Source/Private/OnlineIdentityTwitch.cpp b/Engine/Plugins/Online/OnlineSubsystemTwitch/Source/Private/OnlineIdentityTwitch.cpp
--- a/Engine/Plugins/Online/OnlineSubsystemTwitch/Source/Private/OnlineIdentityTwitch.cpp
+++ b/Engine/Plugins/Online/OnlineSubsystemTwitch/Source/Private/OnlineIdentityTwitch.cpp
@@ -3,6 +3,7 @@
 #include "OnlineSubsystemTwitchPrivate.h"
 #include "OnlineIdentityTwitch.h"
 #include "OnlineAccountTwitch.h"
+#include "OnlineExternalUIInterface.h"
 #include "Http.h"
 #include "Misc/ConfigCacheIni.h"
 
diff --git a/Engine/Plugins/Online/OnlineSubsystemTwitch/Source/Private/OnlineSubsystemTwitch.cpp b/Engine/Plugins/Online/OnlineSubsystemTwitch/Source/Private/OnlineSubsystemTwitch.cpp
--- a/Engine/Plugins/Online/OnlineSubsystemTwitch/Source/Private/OnlineSubsystemTwitch.cpp
+++ b/Engine/Plugins/Online/OnlineSubsystemTwitch/Source/Private/OnlineSubsystemTwitch.cpp
@@ -4,6 +4,7 @@
 #include "OnlineIdentityTwitch.h"
 #include "OnlineExternalUITwitch.h"
 #include "Engine/Console.h"
+#include "Misc/ConfigCacheIni.h"
 
 IOnlineSessionPtr FOnlineSubsystemTwitch::GetSessionInterface() const
 {
