Copying //GamePlugins/Main to none (Source: //GamePlugins/Main/PluginTestGame/Plugins/Online @ 3588247)
============================
  MAJOR FEATURES & CHANGES
============================

Change 3553468 by [REDACTED_PERSON]

	Add and use a function for getting the current console OSS.

	Previously, IOnlineSubsystem::GetByPlatform() was commonly being used to get either the console OSS or expected to return nullptr indicating desktop platform.  However, it can also return the GooglePlay or IOS online subsystems on mobile platforms.  These changes should make it clearer where we are supporting ANY platform oss and where we are specifically looking to get a console oss. These changes also fix several bugs where mobile platforms were being implicitly treated as console.

Change 3582572 by josh.markiewicz

	Redact the OSS auth token from the logs

Change 3581254 by josh.markiewicz

	UE4 - Finalize Google IOS auth login
	- fixed logout not calling its delegate
	- moved auth token data retrieval to its own function

Change 3581237 by josh.markiewicz

	#UE4 - added cookie deletion for Facebook on logout

Change 3581215 by josh.markiewicz

	#UE4 - added cookie deletion for Google on logout

Change 3575510 by ryan.gerleve

	Add additional net stats to help track server performance issues.

Change 3574832 by ryan.gerleve

	Make UIpNetDriver's socket send and receive buffer sizes configurable.

Change 3570966 by rob.cannaday

	Have ToString return Unknown for unknown Party values instead of blank string

Change 3562099 by josh.markiewicz

	#UE4 Encryption token/ack changes
	- moved encryption token request/ack to delegates
	- moved FNetworkNotify to NetworkDelegates.h
	- moved connection logic out of GameInstance and back into networking code
	-- GameInstance sends an enum and the network code does the right thing based on that

Change 3557373 by josh.markiewicz

	#UE4 - only reload config during "restore ini" for classes that were actually affected by the previous hotfix
	- prevents ReloadConfig on every class in an ini file
	- faster parsing in practice (4 reloads vs 1900)

Change 3552190 by dan.hertzka

	Rejoin no longer triggers split-second transitions during the rejoin check
	- Rejoin also properly preloads both the DraftLobby and HUD states (before the map travel, so we'll have it once the loading screen comes down in-game)

Change 3551841 by josh.markiewicz

	fix hotfix manager skipping hotfixes on Xbox
	- if user isn't logged in and its the initial update, skip patch check
	-- xbox shouldn't need it anyway, as patches are required before launch

Change 3545509 by rob.cannaday

	Remove caching of config values in OnlineSubsystemTwitch
	Fixes hotfixes not taking effect

Change 3544669 by josh.markiewicz

	#UE4 - restore debug code for printing datacenter print analytics events
	- disabled by default

Change 3525654 by rob.cannaday

	Add a function to set the StatePrefix programmatically
	#twitch

diff --git a/Engine/Plugins/Online/OnlineFramework/Source/Hotfix/Private/OnlineHotfixManager.cpp b/Engine/Plugins/Online/OnlineFramework/Source/Hotfix/Private/OnlineHotfixManager.cpp
--- a/Engine/Plugins/Online/OnlineFramework/Source/Hotfix/Private/OnlineHotfixManager.cpp
+++ b/Engine/Plugins/Online/OnlineFramework/Source/Hotfix/Private/OnlineHotfixManager.cpp
@@ -641,7 +641,7 @@ bool UOnlineHotfixManager::HotfixIniFile(const FString& FileName, const FString&
 {
 	FConfigFile* ConfigFile = GetConfigFile(FileName);
 	// Store the original file so we can undo this later
-	BackupIniFile(FileName, ConfigFile);
+	FConfigFileBackup& BackupFile = BackupIniFile(FileName, ConfigFile);
 	// Merge the string into the config file
 	ConfigFile->CombineFromBuffer(IniData);
 	TArray<UClass*> Classes;
@@ -676,6 +676,7 @@ bool UOnlineHotfixManager::HotfixIniFile(const FString& FileName, const FString&
 						{
 							// Add this to the list to check against
 							Classes.Add(Class);
+							BackupFile.ClassesReloaded.AddUnique(Class->GetPathName());
 						}
 					}
 					else if (FileName.Contains(TEXT("Engine.ini")))
@@ -706,6 +707,7 @@ bool UOnlineHotfixManager::HotfixIniFile(const FString& FileName, const FString&
 						if (PerObject != nullptr)
 						{
 							PerObjectConfigObjects.Add(PerObject);
+							BackupFile.ClassesReloaded.AddUnique(ObjectClass->GetPathName());
 						}
 					}
 					else
@@ -734,6 +736,7 @@ bool UOnlineHotfixManager::HotfixIniFile(const FString& FileName, const FString&
 					if (It->IsA(Classes[ClassIndex]))
 					{
 						// Force a reload of the config vars
+						UE_LOG(LogHotfixManager, Verbose, TEXT("Reloading %s"), *It->GetPathName());
 						It->ReloadConfig();
 						NumObjectsReloaded++;
 						break;
@@ -745,6 +748,7 @@ bool UOnlineHotfixManager::HotfixIniFile(const FString& FileName, const FString&
 	// Reload any PerObjectConfig objects that were affected
 	for (auto ReloadObject : PerObjectConfigObjects)
 	{
+		UE_LOG(LogHotfixManager, Verbose, TEXT("Reloading %s"), *ReloadObject->GetPathName());
 		ReloadObject->ReloadConfig();
 		NumObjectsReloaded++;
 	}
@@ -936,13 +940,15 @@ void UOnlineHotfixManager::OnReadFileProgress(const FString& FileName, uint64 By
 	}
 }
 
-void UOnlineHotfixManager::BackupIniFile(const FString& IniName, const FConfigFile* ConfigFile)
+UOnlineHotfixManager::FConfigFileBackup& UOnlineHotfixManager::BackupIniFile(const FString& IniName, const FConfigFile* ConfigFile)
 {
 	int32 AddAt = IniBackups.AddDefaulted();
-	IniBackups[AddAt].IniName = GetConfigFileNamePath(GetStrippedConfigFileName(IniName));
-	IniBackups[AddAt].ConfigData = *ConfigFile;
+	FConfigFileBackup& NewBackup = IniBackups[AddAt];
+	NewBackup.IniName = GetConfigFileNamePath(GetStrippedConfigFileName(IniName));
+	NewBackup.ConfigData = *ConfigFile;
 	// There's a lack of deep copy related to the SourceConfigFile so null it out
-	IniBackups[AddAt].ConfigData.SourceConfigFile = nullptr;
+	NewBackup.ConfigData.SourceConfigFile = nullptr;
+	return NewBackup;
 }
 
 void UOnlineHotfixManager::RestoreBackupIniFiles()
@@ -952,7 +958,8 @@ void UOnlineHotfixManager::RestoreBackupIniFiles()
 		return;
 	}
 	const double StartTime = FPlatformTime::Seconds();
-	TArray<FString> RestoredInis;
+	TArray<FString> ClassesToRestore;
+
 	// Restore any changed INI files and build a list of which ones changed for UObject reloading below
 	for (auto& FileHeader : ChangedHotfixFileList)
 	{
@@ -961,16 +968,19 @@ void UOnlineHotfixManager::RestoreBackupIniFiles()
 			const FString ProcessedName = GetConfigFileNamePath(GetStrippedConfigFileName(FileHeader.FileName));
 			for (int32 Index = 0; Index < IniBackups.Num(); Index++)
 			{
+				const FConfigFileBackup& BackupFile = IniBackups[Index];
 				if (IniBackups[Index].IniName == ProcessedName)
 				{
-					GConfig->SetFile(IniBackups[Index].IniName, &IniBackups[Index].ConfigData);
+					ClassesToRestore.Append(BackupFile.ClassesReloaded);
+
+					GConfig->SetFile(BackupFile.IniName, &BackupFile.ConfigData);
 					IniBackups.RemoveAt(Index);
-					RestoredInis.Add(ProcessedName);
 					break;
 				}
 			}
 		}
 	}
+
 	// Also restore any files that were previously part of the hotfix and now are not
 	for (auto& FileHeader : RemovedHotfixFileList)
 	{
@@ -979,29 +989,45 @@ void UOnlineHotfixManager::RestoreBackupIniFiles()
 			const FString ProcessedName = GetConfigFileNamePath(GetStrippedConfigFileName(FileHeader.FileName));
 			for (int32 Index = 0; Index < IniBackups.Num(); Index++)
 			{
-				if (IniBackups[Index].IniName == ProcessedName)
+				const FConfigFileBackup& BackupFile = IniBackups[Index];
+				if (BackupFile.IniName == ProcessedName)
 				{
-					GConfig->SetFile(IniBackups[Index].IniName, &IniBackups[Index].ConfigData);
+					ClassesToRestore.Append(BackupFile.ClassesReloaded);
+
+					GConfig->SetFile(BackupFile.IniName, &BackupFile.ConfigData);
 					IniBackups.RemoveAt(Index);
-					RestoredInis.Add(ProcessedName);
 					break;
 				}
 			}
 		}
 	}
-	// Finally reload any UObjects that were affected by this
+
 	uint32 NumObjectsReloaded = 0;
-	if (RestoredInis.Num() > 0)
+	if (ClassesToRestore.Num() > 0)
 	{
+		TArray<UClass*> RestoredClasses;
+		RestoredClasses.Reserve(ClassesToRestore.Num());
+		for (int32 Index = 0; Index < ClassesToRestore.Num(); Index++)
+		{
+			UClass* Class = FindObject<UClass>(nullptr, *ClassesToRestore[Index], true);
+			if (Class != nullptr)
+			{
+				// Add this to the list to check against
+				RestoredClasses.Add(Class);
+			}
+		}
+
 		for (FObjectIterator It; It; ++It)
 		{
 			UClass* Class = It->GetClass();
 			if (Class->HasAnyClassFlags(CLASS_Config))
 			{
-				for (int32 Index = 0; Index < RestoredInis.Num(); Index++)
+				// Check to see if this class is in our list (yes, potentially n^2, but not in practice)
+				for (int32 ClassIndex = 0; ClassIndex < RestoredClasses.Num(); ClassIndex++)
 				{
-					if (Class->GetConfigName() == RestoredInis[Index])
+					if (It->IsA(RestoredClasses[ClassIndex]))
 					{
+						UE_LOG(LogHotfixManager, Verbose, TEXT("Restoring %s"), *It->GetPathName());
 						It->ReloadConfig();
 						NumObjectsReloaded++;
 						break;
@@ -1010,8 +1036,8 @@ void UOnlineHotfixManager::RestoreBackupIniFiles()
 			}
 		}
 	}
-	UE_LOG(LogHotfixManager, Log, TEXT("Restoring config for %d changed INI files took %f seconds reloading %d objects"),
-		RestoredInis.Num(), FPlatformTime::Seconds() - StartTime, NumObjectsReloaded);
+	UE_LOG(LogHotfixManager, Log, TEXT("Restoring config for %d changed classes took %f seconds reloading %d objects"),
+		ClassesToRestore.Num(), FPlatformTime::Seconds() - StartTime, NumObjectsReloaded);
 }
 
 struct FHotfixManagerExec :
diff --git a/Engine/Plugins/Online/OnlineFramework/Source/Hotfix/Private/UpdateManager.cpp b/Engine/Plugins/Online/OnlineFramework/Source/Hotfix/Private/UpdateManager.cpp
--- a/Engine/Plugins/Online/OnlineFramework/Source/Hotfix/Private/UpdateManager.cpp
+++ b/Engine/Plugins/Online/OnlineFramework/Source/Hotfix/Private/UpdateManager.cpp
@@ -229,19 +229,25 @@ void UUpdateManager::StartPatchCheck()
 
 	SetUpdateState(EUpdateState::CheckingForPatch);
 
-	IOnlineSubsystem* OnlineSubConsole = IOnlineSubsystem::GetByPlatform();
-	if (OnlineSubConsole)
+	IOnlineSubsystem* PlatformOnlineSub = IOnlineSubsystem::GetByPlatform();
+	if (PlatformOnlineSub)
 	{
-		IOnlineIdentityPtr OnlineIdentityConsole = OnlineSubConsole->GetIdentityInterface();
-		if (OnlineIdentityConsole.IsValid())
+		IOnlineIdentityPtr PlatformOnlineIdentity = PlatformOnlineSub->GetIdentityInterface();
+		if (PlatformOnlineIdentity.IsValid())
 		{
-			TSharedPtr<const FUniqueNetId> UserId = GetFirstSignedInUser(OnlineIdentityConsole);
-			if (UserId.IsValid())
+			TSharedPtr<const FUniqueNetId> UserId = GetFirstSignedInUser(PlatformOnlineIdentity);
+			if (UserId.IsValid() && (PlatformOnlineIdentity->GetLoginStatus(*UserId) == ELoginStatus::LoggedIn))
 			{
 				bStarted = true;
-				OnlineIdentityConsole->GetUserPrivilege(*UserId,
+				PlatformOnlineIdentity->GetUserPrivilege(*UserId,
 					EUserPrivileges::CanPlayOnline, IOnlineIdentity::FOnGetUserPrivilegeCompleteDelegate::CreateUObject(this, &ThisClass::OnCheckForPatchComplete, true));
 			}
+			else if (!bInitialUpdateFinished)
+			{
+				UE_LOG(LogHotfixManager, Verbose, TEXT("Skipping initial patch check with no signed in user"));
+				bStarted = true;
+				PatchCheckComplete(EPatchCheckResult::NoPatchRequired);
+			}
 			else
 			{
 				UE_LOG(LogHotfixManager, Warning, TEXT("No valid platform user id when starting patch check!"));
diff --git a/Engine/Plugins/Online/OnlineFramework/Source/Hotfix/Public/OnlineHotfixManager.h b/Engine/Plugins/Online/OnlineFramework/Source/Hotfix/Public/OnlineHotfixManager.h
--- a/Engine/Plugins/Online/OnlineFramework/Source/Hotfix/Public/OnlineHotfixManager.h
+++ b/Engine/Plugins/Online/OnlineFramework/Source/Hotfix/Public/OnlineHotfixManager.h
@@ -119,8 +119,12 @@ protected:
 
 	struct FConfigFileBackup
 	{
+		/** Name of the ini file backed up*/
 		FString IniName;
+		/** Previous ini data backed up */
 		FConfigFile ConfigData;
+		/** UClasses reloaded as a result of the current ini */
+		TArray<FString> ClassesReloaded;
 	};
 
 	/** Holds which files are pending download */
@@ -180,7 +184,7 @@ protected:
 	/** Unmounts any changed PAK files so they can be re-mounted after downloading */
 	void UnmountHotfixFiles();
 	/** Stores off the INI file for restoration later */
-	void BackupIniFile(const FString& IniName, const FConfigFile* ConfigFile);
+	FConfigFileBackup& BackupIniFile(const FString& IniName, const FConfigFile* ConfigFile);
 	/** Restores any changed INI files to their default loaded state */
 	void RestoreBackupIniFiles();
 	/** Builds the list of files that are different between two runs of the hotfix process */
diff --git a/Engine/Plugins/Online/OnlineFramework/Source/LoginFlow/Private/SLoginFlow.cpp b/Engine/Plugins/Online/OnlineFramework/Source/LoginFlow/Private/SLoginFlow.cpp
--- a/Engine/Plugins/Online/OnlineFramework/Source/LoginFlow/Private/SLoginFlow.cpp
+++ b/Engine/Plugins/Online/OnlineFramework/Source/LoginFlow/Private/SLoginFlow.cpp
@@ -355,12 +355,12 @@ private:
 
 				// Setup browser widget.
 				TSharedPtr<SWebBrowser> BrowserWidget;
-					BrowserWindowWidget->SetContent(
-						SNew(SBorder)
-						.VAlign(VAlign_Fill)
-						.HAlign(HAlign_Fill)
-						.Padding(0)
-						[
+				BrowserWindowWidget->SetContent(
+					SNew(SBorder)
+					.VAlign(VAlign_Fill)
+					.HAlign(HAlign_Fill)
+					.Padding(0)
+					[
 						SAssignNew(BrowserWidget, SWebBrowser, NewBrowserWindowSP)
 						.ShowControls(PopupFeaturesSP->IsToolBarVisible())
 						.ShowAddressBar(PopupFeaturesSP->IsLocationBarVisible())
@@ -370,38 +370,38 @@ private:
 					]);
 
 				// Setup some OnClose stuff.
-			{
-				struct FLocal
 				{
-					static void RequestDestroyWindowOverride(const TSharedRef<SWindow>& Window, TWeakPtr<IWebBrowserWindow> BrowserWindowPtr)
+					struct FLocal
 					{
-						TSharedPtr<IWebBrowserWindow> BrowserWindow = BrowserWindowPtr.Pin();
-						if (BrowserWindow.IsValid())
+						static void RequestDestroyWindowOverride(const TSharedRef<SWindow>& Window, TWeakPtr<IWebBrowserWindow> BrowserWindowPtr)
 						{
-							if (BrowserWindow->IsClosing())
-							{
-								FSlateApplicationBase::Get().RequestDestroyWindow(Window);
-							}
-							else
+							TSharedPtr<IWebBrowserWindow> BrowserWindow = BrowserWindowPtr.Pin();
+							if (BrowserWindow.IsValid())
 							{
-								// Notify the browser window that we would like to close it.  On the CEF side, this will 
-								//  result in a call to FWebBrowserHandler::DoClose only if the JavaScript onbeforeunload
-								//  event handler allows it.
-								BrowserWindow->CloseBrowser(false);
+								if (BrowserWindow->IsClosing())
+								{
+									FSlateApplicationBase::Get().RequestDestroyWindow(Window);
+								}
+								else
+								{
+									// Notify the browser window that we would like to close it.  On the CEF side, this will 
+									//  result in a call to FWebBrowserHandler::DoClose only if the JavaScript onbeforeunload
+									//  event handler allows it.
+									BrowserWindow->CloseBrowser(false);
+								}
 							}
 						}
-					}
-				};
+					};
 
-				BrowserWindowWidget->SetRequestDestroyWindowOverride(FRequestDestroyWindowOverride::CreateStatic(&FLocal::RequestDestroyWindowOverride, TWeakPtr<IWebBrowserWindow>(NewBrowserWindow)));
-			}
+					BrowserWindowWidget->SetRequestDestroyWindowOverride(FRequestDestroyWindowOverride::CreateStatic(&FLocal::RequestDestroyWindowOverride, TWeakPtr<IWebBrowserWindow>(NewBrowserWindow)));
+				}
 
-			FSlateApplication::Get().AddWindow(BrowserWindowWidget);
-			BrowserWindowWidget->BringToFront();
-			FSlateApplication::Get().SetKeyboardFocus(BrowserWidget, EFocusCause::SetDirectly);
+				FSlateApplication::Get().AddWindow(BrowserWindowWidget);
+				BrowserWindowWidget->BringToFront();
+				FSlateApplication::Get().SetKeyboardFocus(BrowserWidget, EFocusCause::SetDirectly);
 
-			BrowserWindowWidgets.Add(NewBrowserWindow, BrowserWindowWidget);
-			return true;
+				BrowserWindowWidgets.Add(NewBrowserWindow, BrowserWindowWidget);
+				return true;
 			}
 		}
 		else
diff --git a/Engine/Plugins/Online/OnlineFramework/Source/Party/Private/Party.cpp b/Engine/Plugins/Online/OnlineFramework/Source/Party/Private/Party.cpp
--- a/Engine/Plugins/Online/OnlineFramework/Source/Party/Private/Party.cpp
+++ b/Engine/Plugins/Online/OnlineFramework/Source/Party/Private/Party.cpp
@@ -172,43 +172,53 @@ void UParty::AddReferencedObjects(UObject* InThis, FReferenceCollector& Collecto
 
 void UParty::RegisterIdentityDelegates()
 {
-	UWorld* World = GetWorld();
-	IOnlineIdentityPtr IdentityInt = Online::GetIdentityInterface(World);
-	if (IdentityInt.IsValid())
+	UWorld* const World = GetWorld();
+	if (ensure(World))
 	{
-		// Unbind and then rebind
-		UnregisterIdentityDelegates();
+		const IOnlineIdentityPtr IdentityInt = Online::GetIdentityInterface(World);
+		if (IdentityInt.IsValid())
+		{
+			// Unbind and then rebind
+			UnregisterIdentityDelegates();
 
-		FOnLogoutCompleteDelegate OnLogoutCompleteDelegate;
-		OnLogoutCompleteDelegate.BindUObject(this, &ThisClass::OnLogoutComplete);
+			FOnLogoutCompleteDelegate OnLogoutCompleteDelegate;
+			OnLogoutCompleteDelegate.BindUObject(this, &ThisClass::OnLogoutComplete);
 
-		FOnLoginStatusChangedDelegate OnLoginStatusChangedDelegate;
-		OnLoginStatusChangedDelegate.BindUObject(this, &ThisClass::OnLoginStatusChanged);
+			FOnLoginStatusChangedDelegate OnLoginStatusChangedDelegate;
+			OnLoginStatusChangedDelegate.BindUObject(this, &ThisClass::OnLoginStatusChanged);
 
-		for (int32 LocalPlayerId = 0; LocalPlayerId < MAX_LOCAL_PLAYERS; LocalPlayerId++)
-		{
-			LogoutCompleteDelegateHandle[LocalPlayerId] = IdentityInt->AddOnLogoutCompleteDelegate_Handle(LocalPlayerId, OnLogoutCompleteDelegate);
-			LogoutStatusChangedDelegateHandle[LocalPlayerId] = IdentityInt->AddOnLoginStatusChangedDelegate_Handle(LocalPlayerId, OnLoginStatusChangedDelegate);
+			for (int32 LocalPlayerId = 0; LocalPlayerId < MAX_LOCAL_PLAYERS; LocalPlayerId++)
+			{
+				LogoutCompleteDelegateHandle[LocalPlayerId] = IdentityInt->AddOnLogoutCompleteDelegate_Handle(LocalPlayerId, OnLogoutCompleteDelegate);
+				LogoutStatusChangedDelegateHandle[LocalPlayerId] = IdentityInt->AddOnLoginStatusChangedDelegate_Handle(LocalPlayerId, OnLoginStatusChangedDelegate);
+			}
 		}
 	}
+	else
+	{
+		UE_LOG(LogParty, Warning, TEXT("UParty::RegisterIdentityDelegates: Missing World!"));
+	}
 }
 
 void UParty::UnregisterIdentityDelegates()
 {
-	UWorld* World = GetWorld();
-	IOnlineIdentityPtr IdentityInt = Online::GetIdentityInterface(World);
-	if (IdentityInt.IsValid())
+	UWorld* const World = GetWorld();
+	if (World)
 	{
-		for (int32 LocalPlayerId = 0; LocalPlayerId < MAX_LOCAL_PLAYERS; LocalPlayerId++)
+		const IOnlineIdentityPtr IdentityInt = Online::GetIdentityInterface(World);
+		if (IdentityInt.IsValid())
 		{
-			if (LogoutCompleteDelegateHandle[LocalPlayerId].IsValid())
+			for (int32 LocalPlayerId = 0; LocalPlayerId < MAX_LOCAL_PLAYERS; LocalPlayerId++)
 			{
-				IdentityInt->ClearOnLogoutCompleteDelegate_Handle(LocalPlayerId, LogoutCompleteDelegateHandle[LocalPlayerId]);
-			}
+				if (LogoutCompleteDelegateHandle[LocalPlayerId].IsValid())
+				{
+					IdentityInt->ClearOnLogoutCompleteDelegate_Handle(LocalPlayerId, LogoutCompleteDelegateHandle[LocalPlayerId]);
+				}
 
-			if (LogoutStatusChangedDelegateHandle[LocalPlayerId].IsValid())
-			{
-				IdentityInt->ClearOnLoginStatusChangedDelegate_Handle(LocalPlayerId, LogoutStatusChangedDelegateHandle[LocalPlayerId]);
+				if (LogoutStatusChangedDelegateHandle[LocalPlayerId].IsValid())
+				{
+					IdentityInt->ClearOnLoginStatusChangedDelegate_Handle(LocalPlayerId, LogoutStatusChangedDelegateHandle[LocalPlayerId]);
+				}
 			}
 		}
 	}
@@ -216,45 +226,50 @@ void UParty::UnregisterIdentityDelegates()
 
 void UParty::RegisterPartyDelegates()
 {
-	UnregisterPartyDelegates();
-
-	UWorld* World = GetWorld();
-	check(World);
-
-	IOnlinePartyPtr PartyInt = Online::GetPartyInterface(World);
-	if (PartyInt.IsValid())
+	UWorld* const World = GetWorld();
+	if (ensure(World))
 	{
-		PartyConfigChangedDelegateHandle = PartyInt->AddOnPartyConfigChangedDelegate_Handle(FOnPartyConfigChangedDelegate::CreateUObject(this, &ThisClass::PartyConfigChangedInternal));
-		PartyMemberJoinedDelegateHandle = PartyInt->AddOnPartyMemberJoinedDelegate_Handle(FOnPartyMemberJoinedDelegate::CreateUObject(this, &ThisClass::PartyMemberJoinedInternal));
-		PartyDataReceivedDelegateHandle = PartyInt->AddOnPartyDataReceivedDelegate_Handle(FOnPartyDataReceivedDelegate::CreateUObject(this, &ThisClass::PartyDataReceivedInternal));
-		PartyMemberDataReceivedDelegateHandle = PartyInt->AddOnPartyMemberDataReceivedDelegate_Handle(FOnPartyMemberDataReceivedDelegate::CreateUObject(this, &ThisClass::PartyMemberDataReceivedInternal));
-		PartyJoinRequestReceivedDelegateHandle = PartyInt->AddOnPartyJoinRequestReceivedDelegate_Handle(FOnPartyJoinRequestReceivedDelegate::CreateUObject(this, &ThisClass::PartyJoinRequestReceivedInternal));
-		PartyQueryJoinabilityReceivedDelegateHandle = PartyInt->AddOnQueryPartyJoinabilityReceivedDelegate_Handle(FOnPartyJoinRequestReceivedDelegate::CreateUObject(this, &ThisClass::PartyQueryJoinabilityReceivedInternal));
-		PartyMemberPromotedDelegateHandle = PartyInt->AddOnPartyMemberPromotedDelegate_Handle(FOnPartyMemberPromotedDelegate::CreateUObject(this, &ThisClass::PartyMemberPromotedInternal));
-		PartyMemberExitedDelegateHandle = PartyInt->AddOnPartyMemberExitedDelegate_Handle(FOnPartyMemberExitedDelegate::CreateUObject(this, &ThisClass::PartyMemberExitedInternal));
-		PartyPromotionLockoutChangedDelegateHandle = PartyInt->AddOnPartyPromotionLockoutChangedDelegate_Handle(FOnPartyPromotionLockoutChangedDelegate::CreateUObject(this, &ThisClass::PartyPromotionLockoutStateChangedInternal));
-		PartyExitedDelegateHandle = PartyInt->AddOnPartyExitedDelegate_Handle(FOnPartyExitedDelegate::CreateUObject(this, &ThisClass::PartyExitedInternal));
-		PartyStateChangedDelegateHandle = PartyInt->AddOnPartyStateChangedDelegate_Handle(FOnPartyStateChangedDelegate::CreateUObject(this, &ThisClass::PartyStateChanged));
+		const IOnlinePartyPtr PartyInt = Online::GetPartyInterface(World);
+		if (PartyInt.IsValid())
+		{
+			// Unbind and then rebind
+			UnregisterPartyDelegates();
+
+			PartyConfigChangedDelegateHandle = PartyInt->AddOnPartyConfigChangedDelegate_Handle(FOnPartyConfigChangedDelegate::CreateUObject(this, &ThisClass::PartyConfigChangedInternal));
+			PartyMemberJoinedDelegateHandle = PartyInt->AddOnPartyMemberJoinedDelegate_Handle(FOnPartyMemberJoinedDelegate::CreateUObject(this, &ThisClass::PartyMemberJoinedInternal));
+			PartyDataReceivedDelegateHandle = PartyInt->AddOnPartyDataReceivedDelegate_Handle(FOnPartyDataReceivedDelegate::CreateUObject(this, &ThisClass::PartyDataReceivedInternal));
+			PartyMemberDataReceivedDelegateHandle = PartyInt->AddOnPartyMemberDataReceivedDelegate_Handle(FOnPartyMemberDataReceivedDelegate::CreateUObject(this, &ThisClass::PartyMemberDataReceivedInternal));
+			PartyJoinRequestReceivedDelegateHandle = PartyInt->AddOnPartyJoinRequestReceivedDelegate_Handle(FOnPartyJoinRequestReceivedDelegate::CreateUObject(this, &ThisClass::PartyJoinRequestReceivedInternal));
+			PartyQueryJoinabilityReceivedDelegateHandle = PartyInt->AddOnQueryPartyJoinabilityReceivedDelegate_Handle(FOnPartyJoinRequestReceivedDelegate::CreateUObject(this, &ThisClass::PartyQueryJoinabilityReceivedInternal));
+			PartyMemberPromotedDelegateHandle = PartyInt->AddOnPartyMemberPromotedDelegate_Handle(FOnPartyMemberPromotedDelegate::CreateUObject(this, &ThisClass::PartyMemberPromotedInternal));
+			PartyMemberExitedDelegateHandle = PartyInt->AddOnPartyMemberExitedDelegate_Handle(FOnPartyMemberExitedDelegate::CreateUObject(this, &ThisClass::PartyMemberExitedInternal));
+			PartyPromotionLockoutChangedDelegateHandle = PartyInt->AddOnPartyPromotionLockoutChangedDelegate_Handle(FOnPartyPromotionLockoutChangedDelegate::CreateUObject(this, &ThisClass::PartyPromotionLockoutStateChangedInternal));
+			PartyExitedDelegateHandle = PartyInt->AddOnPartyExitedDelegate_Handle(FOnPartyExitedDelegate::CreateUObject(this, &ThisClass::PartyExitedInternal));
+			PartyStateChangedDelegateHandle = PartyInt->AddOnPartyStateChangedDelegate_Handle(FOnPartyStateChangedDelegate::CreateUObject(this, &ThisClass::PartyStateChanged));
+		}
 	}
 }
 
 void UParty::UnregisterPartyDelegates()
 {
-	UWorld* World = GetWorld();
-	IOnlinePartyPtr PartyInt = Online::GetPartyInterface(World);
-	if (PartyInt.IsValid())
+	UWorld* const World = GetWorld();
+	if (World)
 	{
-		PartyInt->ClearOnPartyConfigChangedDelegate_Handle(PartyConfigChangedDelegateHandle);
-		PartyInt->ClearOnPartyMemberJoinedDelegate_Handle(PartyMemberJoinedDelegateHandle);
-		PartyInt->ClearOnPartyDataReceivedDelegate_Handle(PartyDataReceivedDelegateHandle);
-		PartyInt->ClearOnPartyMemberDataReceivedDelegate_Handle(PartyMemberDataReceivedDelegateHandle);
-		PartyInt->ClearOnPartyJoinRequestReceivedDelegate_Handle(PartyJoinRequestReceivedDelegateHandle);
-		PartyInt->ClearOnQueryPartyJoinabilityReceivedDelegate_Handle(PartyQueryJoinabilityReceivedDelegateHandle);
-		PartyInt->ClearOnPartyMemberPromotedDelegate_Handle(PartyMemberPromotedDelegateHandle);
-		PartyInt->ClearOnPartyMemberExitedDelegate_Handle(PartyMemberExitedDelegateHandle);
-		PartyInt->ClearOnPartyPromotionLockoutChangedDelegate_Handle(PartyPromotionLockoutChangedDelegateHandle);
-		PartyInt->ClearOnPartyExitedDelegate_Handle(PartyExitedDelegateHandle);
-		PartyInt->ClearOnPartyStateChangedDelegate_Handle(PartyStateChangedDelegateHandle);
+		const IOnlinePartyPtr PartyInt = Online::GetPartyInterface(World);
+		if (PartyInt.IsValid())
+		{
+			PartyInt->ClearOnPartyConfigChangedDelegate_Handle(PartyConfigChangedDelegateHandle);
+			PartyInt->ClearOnPartyMemberJoinedDelegate_Handle(PartyMemberJoinedDelegateHandle);
+			PartyInt->ClearOnPartyDataReceivedDelegate_Handle(PartyDataReceivedDelegateHandle);
+			PartyInt->ClearOnPartyMemberDataReceivedDelegate_Handle(PartyMemberDataReceivedDelegateHandle);
+			PartyInt->ClearOnPartyJoinRequestReceivedDelegate_Handle(PartyJoinRequestReceivedDelegateHandle);
+			PartyInt->ClearOnQueryPartyJoinabilityReceivedDelegate_Handle(PartyQueryJoinabilityReceivedDelegateHandle);
+			PartyInt->ClearOnPartyMemberPromotedDelegate_Handle(PartyMemberPromotedDelegateHandle);
+			PartyInt->ClearOnPartyMemberExitedDelegate_Handle(PartyMemberExitedDelegateHandle);
+			PartyInt->ClearOnPartyPromotionLockoutChangedDelegate_Handle(PartyPromotionLockoutChangedDelegateHandle);
+			PartyInt->ClearOnPartyExitedDelegate_Handle(PartyExitedDelegateHandle);
+			PartyInt->ClearOnPartyStateChangedDelegate_Handle(PartyStateChangedDelegateHandle);
+		}
 	}
 }
 
diff --git a/Engine/Plugins/Online/OnlineFramework/Source/Party/Public/PartyGameState.h b/Engine/Plugins/Online/OnlineFramework/Source/Party/Public/PartyGameState.h
--- a/Engine/Plugins/Online/OnlineFramework/Source/Party/Public/PartyGameState.h
+++ b/Engine/Plugins/Online/OnlineFramework/Source/Party/Public/PartyGameState.h
@@ -720,7 +720,7 @@ inline const TCHAR* ToString(EPartyType Type)
 	}
 	default:
 	{
-		return TEXT("");
+		return TEXT("Unknown");
 	}
 	}
 }
@@ -759,7 +759,7 @@ inline const TCHAR* ToString(EJoinPartyDenialReason Type)
 	}
 	default:
 	{
-		return TEXT("");
+		return TEXT("Unknown");
 	}
 	}
 }
@@ -786,7 +786,7 @@ inline const TCHAR* ToString(EApprovalAction Type)
 	}
 	default:
 	{
-		return TEXT("");
+		return TEXT("Unknown");
 	}
 	}
 }
diff --git a/Engine/Plugins/Online/OnlineFramework/Source/Qos/Private/QosStats.cpp b/Engine/Plugins/Online/OnlineFramework/Source/Qos/Private/QosStats.cpp
--- a/Engine/Plugins/Online/OnlineFramework/Source/Qos/Private/QosStats.cpp
+++ b/Engine/Plugins/Online/OnlineFramework/Source/Qos/Private/QosStats.cpp
@@ -7,6 +7,7 @@
 #include "Interfaces/IAnalyticsProvider.h"
 
 #define QOS_STATS_VERSION 1
+#define DEBUG_QOS_STATS 0
 
 // Events
 const FString FQosDatacenterStats::QosStats_DatacenterEvent = TEXT("QosStats_DatacenterEvent");
@@ -27,6 +28,24 @@ const FString FQosDatacenterStats::QosStats_NumResults = TEXT("QosStats_NumResul
 const FString FQosDatacenterStats::QosStats_NumSuccessCount = TEXT("QosStats_NumSuccessCount");
 const FString FQosDatacenterStats::QosStats_SearchDetails = TEXT("QosStats_SearchDetails");
 
+/**
+ * Debug output for the contents of a recorded stats event
+ *
+ * @param StatsEvent event type recorded
+ * @param Attributes attribution of the event
+ */
+inline void PrintEventAndAttributes(const FString& StatsEvent, const TArray<FAnalyticsEventAttribute>& Attributes)
+{
+#if DEBUG_QOS_STATS
+	UE_LOG(LogQos, Display, TEXT("Event: %s"), *StatsEvent);
+	for (int32 AttrIdx = 0; AttrIdx<Attributes.Num(); AttrIdx++)
+	{
+		const FAnalyticsEventAttribute& Attr = Attributes[AttrIdx];
+		UE_LOG(LogQos, Display, TEXT("\t%s : %s"), *Attr.AttrName, *Attr.AttrValue);
+	}
+#endif
+}
+
 FQosDatacenterStats::FQosDatacenterStats() :
 	StatsVersion(QOS_STATS_VERSION),
 	bAnalyticsInProgress(false)
@@ -200,6 +219,7 @@ void FQosDatacenterStats::ParseQosResults(TSharedPtr<IAnalyticsProvider>& Analyt
 	}
 	QoSAttributes.Add(FAnalyticsEventAttribute(QosStats_SearchDetails, SearchDetails));
 
+	PrintEventAndAttributes(QosStats_DatacenterEvent, QoSAttributes);
 	if (AnalyticsProvider.IsValid())
 	{
 		AnalyticsProvider->RecordEvent(QosStats_DatacenterEvent, QoSAttributes);
diff --git a/Engine/Plugins/Online/OnlineFramework/Source/Rejoin/Private/RejoinCheck.cpp b/Engine/Plugins/Online/OnlineFramework/Source/Rejoin/Private/RejoinCheck.cpp
--- a/Engine/Plugins/Online/OnlineFramework/Source/Rejoin/Private/RejoinCheck.cpp
+++ b/Engine/Plugins/Online/OnlineFramework/Source/Rejoin/Private/RejoinCheck.cpp
@@ -43,6 +43,19 @@ ERejoinStatus URejoinCheck::GetStatus() const
 }
 #endif
 
+bool URejoinCheck::HasCompletedCheck() const
+{
+	ERejoinStatus CurrentStatus = GetStatus();
+	return (CurrentStatus != ERejoinStatus::NeedsRecheck && CurrentStatus != ERejoinStatus::UpdatingStatus);
+}
+
+bool URejoinCheck::IsRejoinAvailable() const
+{
+	ERejoinStatus CurrentStatus = GetStatus();
+	return (CurrentStatus != ERejoinStatus::NoMatchToRejoin &&
+		CurrentStatus != ERejoinStatus::NoMatchToRejoin_MatchEnded);
+}
+
 void URejoinCheck::CheckRejoinStatus(const FOnRejoinCheckComplete& InCompletionDelegate)
 {
 #if !UE_BUILD_SHIPPING
@@ -301,6 +314,8 @@ void URejoinCheck::TravelToSession()
 		bResult = GameInstance->ClientTravelToSession(LP->GetControllerId(), GameSessionName);
 		if (bResult)
 		{
+			UE_LOG(LogOnline, Log, TEXT("URejoinCheck::TravelToSession: Performing ClientTravelToSession"));
+
 			// Record the result of the attempt to rejoin
 			Analytics_RecordRejoinAttempt(SearchResult, ERejoinAttemptResult::RejoinSuccess);
 
@@ -310,12 +325,12 @@ void URejoinCheck::TravelToSession()
 		}
 		else
 		{
-			UE_LOG(LogOnline, Verbose, TEXT("TravelToSession: Failed to travel to session"));
+			UE_LOG(LogOnline, Log, TEXT("URejoinCheck::TravelToSession: Failed to travel to session"));
 		}
 	}
 	else
 	{
-		UE_LOG(LogOnline, Verbose, TEXT("TravelToSession: Failed to find local player"));
+		UE_LOG(LogOnline, Log, TEXT("URejoinCheck::TravelToSession: Failed to find local player"));
 	}
 	
 	if (!bResult)
diff --git a/Engine/Plugins/Online/OnlineFramework/Source/Rejoin/Public/RejoinCheck.h b/Engine/Plugins/Online/OnlineFramework/Source/Rejoin/Public/RejoinCheck.h
--- a/Engine/Plugins/Online/OnlineFramework/Source/Rejoin/Public/RejoinCheck.h
+++ b/Engine/Plugins/Online/OnlineFramework/Source/Rejoin/Public/RejoinCheck.h
@@ -131,13 +131,11 @@ public:
 	ERejoinStatus GetStatus() const;
 #endif
 
-	/** @return true if possibly a match to rejoin */
-	bool IsRejoinAvailable() const
-	{
-		ERejoinStatus CurrentStatus = GetStatus();
-		return (CurrentStatus != ERejoinStatus::NoMatchToRejoin &&
-				CurrentStatus != ERejoinStatus::NoMatchToRejoin_MatchEnded);
-	}
+	/** @return True if the rejoin check has completed and does not need to be rerun. */
+	bool HasCompletedCheck() const;
+
+	/** @return True if it's possible that there's a match to rejoin */
+	bool IsRejoinAvailable() const;
 
 protected: 
 
diff --git a/Engine/Plugins/Online/OnlineSubsystem/Source/Private/OnlineSessionSettings.cpp b/Engine/Plugins/Online/OnlineSubsystem/Source/Private/OnlineSessionSettings.cpp
--- a/Engine/Plugins/Online/OnlineSubsystem/Source/Private/OnlineSessionSettings.cpp
+++ b/Engine/Plugins/Online/OnlineSubsystem/Source/Private/OnlineSessionSettings.cpp
@@ -184,7 +184,7 @@ int32 FOnlineSessionSettings::GetID(FName Key) const
 	}
 
 	UE_LOG(LogOnline, Warning, TEXT("Unable to find key for ID request: %s"), *Key.ToString());
-	return -1;
+	return INVALID_SESSION_SETTING_ID;
 }
 
 template<typename ValueType>
diff --git a/Engine/Plugins/Online/OnlineSubsystem/Source/Public/Interfaces/OnlineCodeRedemptionInterface.h b/Engine/Plugins/Online/OnlineSubsystem/Source/Public/Interfaces/OnlineCodeRedemptionInterface.h
new file mode 100644
--- /dev/null
+++ b/Engine/Plugins/Online/OnlineSubsystem/Source/Public/Interfaces/OnlineCodeRedemptionInterface.h
@@ -0,0 +1,84 @@
+// Copyright 1998-2017 Epic Games, Inc. All Rights Reserved.
+
+#pragma once
+
+#include "CoreMinimal.h"
+
+struct FOnlineError;
+
+class FCodeLockInfo
+{
+public:
+	FString CodeType;
+	FString CodeId;
+	FString CodeUseId;
+	FString OfferId;
+};
+
+class FCodeEvaluationInfo
+{
+public:
+	FString CodeType;
+	FString CodeId;
+	FString CodeStatus;
+	FString OfferId;
+};
+
+/**
+ * Delegate used when the a lock code request is completed
+ *
+ * @param FOnlineError The Result of the HttpRequest
+ * @param TSharedRef<FCodeLockInfo> Data From a successful Request
+ */
+DECLARE_DELEGATE_TwoParams(FOnProductCodeLockedComplete, const FOnlineError& /*Result*/, const TSharedRef<FCodeLockInfo>&);
+
+/**
+ * Delegate used when the an evaluate code request is completed
+ *
+ * @param FOnlineError The Result of the HttpRequest
+ */
+DECLARE_DELEGATE_OneParam(FOnProductCodeUnlockedComplete, const FOnlineError& /*Result*/);
+
+/**
+ * Delegate used when the an evaluate code request is completed
+ *
+ * @param FOnlineError The Result of the HttpRequest
+ * @param TSharedRef<FCodeEvaluationInfo> Data From a successful Request
+ */
+DECLARE_DELEGATE_TwoParams(FOnProductCodeEvaluateComplete, const FOnlineError& /*Result*/, const TSharedRef<FCodeEvaluationInfo>&);
+
+/**
+ *	IOnlineCodeRedemption - Interface for locking, unlocking, and evaluating codes. Use IOnlineFulfillment to redeem it.
+ */
+class IOnlineCodeRedemption
+{
+public:
+	virtual ~IOnlineCodeRedemption() {}
+
+	/**
+	 * Initiate the lock code process for reserving a product redemption
+	 *
+	 * @param CodeId The Redemption code to lock
+	 * @param LockTimeoutSeconds The duration in seconds to keep the code locked for
+	 * @param Delegate completion callback (guaranteed to be called)
+	 */
+	virtual void LockCode(const FString& CodeId, int32 LockTimeoutSeconds, const FOnProductCodeLockedComplete& Delegate) = 0;
+
+	/**
+	 * Initiate the Release code process to clear reservation on a product code
+	 *
+	 * @param CodeId The Redemption code to unlock
+	 * @param CodeUseId The in use Code Id that was returned from LockCode
+	 * @param Delegate completion callback (guaranteed to be called)
+	 */
+	virtual void UnlockCode(const FString& CodeId, const FString& CodeUseId, const FOnProductCodeUnlockedComplete& Delegate) = 0;
+
+	/**
+	 * Initiate the Evaluation Code process for checking if a code is valid and what product information it relates too
+	 *
+	 * @param UserId User initiating the request
+	 * @param CodeId The Redemption code to evaluate
+	 * @param Delegate completion callback (guaranteed to be called)
+	 */
+	virtual void EvaluateCode(const FUniqueNetId& UserId, const FString& CodeId, const FOnProductCodeEvaluateComplete& Delegate) = 0;
+};
diff --git a/Engine/Plugins/Online/OnlineSubsystem/Source/Public/Interfaces/OnlineFulfillmentInterface.h b/Engine/Plugins/Online/OnlineSubsystem/Source/Public/Interfaces/OnlineFulfillmentInterface.h
new file mode 100644
--- /dev/null
+++ b/Engine/Plugins/Online/OnlineSubsystem/Source/Public/Interfaces/OnlineFulfillmentInterface.h
@@ -0,0 +1,24 @@
+// Copyright 1998-2017 Epic Games, Inc. All Rights Reserved.
+
+#pragma once
+
+#include "CoreMinimal.h"
+#include "OnlinePurchaseInterface.h"
+
+/**
+ * IOnlineFulfillment - Interface for redeeming codes.
+ */
+class IOnlineFulfillment
+{
+public:
+	virtual ~IOnlineFulfillment() {}
+
+	/**
+	 * Initiate the checkout process for obtaining offers via code redemption
+	 *
+	 * @param UserId user initiating the request
+	 * @param RedeemCodeRequest info needed for the redeem request
+	 * @param Delegate completion callback (guaranteed to be called)
+	 */
+	virtual void RedeemCode(const FUniqueNetId& UserId, const FRedeemCodeRequest& RedeemCodeRequest, const FOnPurchaseRedeemCodeComplete& Delegate) = 0;
+};
diff --git a/Engine/Plugins/Online/OnlineSubsystem/Source/Public/Interfaces/OnlineIdentityInterface.h b/Engine/Plugins/Online/OnlineSubsystem/Source/Public/Interfaces/OnlineIdentityInterface.h
--- a/Engine/Plugins/Online/OnlineSubsystem/Source/Public/Interfaces/OnlineIdentityInterface.h
+++ b/Engine/Plugins/Online/OnlineSubsystem/Source/Public/Interfaces/OnlineIdentityInterface.h
@@ -236,7 +236,7 @@ public:
 	DEFINE_ONLINE_PLAYER_DELEGATE_ONE_PARAM(MAX_LOCAL_PLAYERS, OnLogoutComplete, bool);
 
 	/**
-	 * Delegate called when the online subsystem requires an the login flow to logout and cleanup
+	 * Delegate called when the online subsystem requires the login flow to logout and cleanup
 	 * @param LoginDomains login domains to be cleaned up
 	 */
 	DEFINE_ONLINE_DELEGATE_ONE_PARAM(OnLoginFlowLogout, const TArray<FString>& /*LoginDomains*/);
diff --git a/Engine/Plugins/Online/OnlineSubsystem/Source/Public/Interfaces/OnlineMessageInterface.h b/Engine/Plugins/Online/OnlineSubsystem/Source/Public/Interfaces/OnlineMessageInterface.h
--- a/Engine/Plugins/Online/OnlineSubsystem/Source/Public/Interfaces/OnlineMessageInterface.h
+++ b/Engine/Plugins/Online/OnlineSubsystem/Source/Public/Interfaces/OnlineMessageInterface.h
@@ -108,7 +108,7 @@ public:
 	/**
 	 * Constructor
 	 */
-	FOnlineMessageHeader(FUniqueNetId* InFromUserId, FUniqueNetId* InMessageId)
+	FOnlineMessageHeader(const TSharedRef<const FUniqueNetId>& InFromUserId, const TSharedRef<FUniqueMessageId>& InMessageId)
 		: FromUserId(InFromUserId)
 		, MessageId(InMessageId)
 	{
diff --git a/Engine/Plugins/Online/OnlineSubsystem/Source/Public/Interfaces/OnlinePartyInterface.h b/Engine/Plugins/Online/OnlineSubsystem/Source/Public/Interfaces/OnlinePartyInterface.h
--- a/Engine/Plugins/Online/OnlineSubsystem/Source/Public/Interfaces/OnlinePartyInterface.h
+++ b/Engine/Plugins/Online/OnlineSubsystem/Source/Public/Interfaces/OnlinePartyInterface.h
@@ -1578,7 +1578,7 @@ inline const TCHAR* ToString(const EPartyState Value)
 		return TEXT("CleanUp");
 	}
 	}
-	return TEXT("");
+	return TEXT("Unknown");
 }
 
 inline const TCHAR* ToString(const EMemberExitedReason Value)
@@ -1602,7 +1602,7 @@ inline const TCHAR* ToString(const EMemberExitedReason Value)
 		return TEXT("Kicked");
 	}
 	}
-	return TEXT("");
+	return TEXT("Unknown"); // Same as EMemberExitedReason::Unknown, which is ok because it is only used when we do not have enough information
 }
 
 inline const TCHAR* ToString(const ECreatePartyCompletionResult Value)
@@ -1642,7 +1642,7 @@ inline const TCHAR* ToString(const ECreatePartyCompletionResult Value)
 		return TEXT("Succeeded");
 	}
 	}
-	return TEXT("");
+	return TEXT("Unknown");
 }
 
 inline const TCHAR* ToString(const ESendPartyInvitationCompletionResult Value)
@@ -1678,7 +1678,7 @@ inline const TCHAR* ToString(const ESendPartyInvitationCompletionResult Value)
 		return TEXT("Succeeded");
 	}
 	}
-	return TEXT("");
+	return TEXT("Unknown");
 }
 
 inline const TCHAR* ToString(const EJoinPartyCompletionResult Value)
@@ -1774,7 +1774,7 @@ inline const TCHAR* ToString(const EJoinPartyCompletionResult Value)
 		return TEXT("DeprecatedUnknownInternalFailure");
 	}
 	}
-	return TEXT("");
+	return TEXT("Unknown");
 }
 
 inline const TCHAR* ToString(const ELeavePartyCompletionResult Value)
@@ -1826,7 +1826,7 @@ inline const TCHAR* ToString(const ELeavePartyCompletionResult Value)
 		return TEXT("DeprecatedUnknownInternalFailure");
 	}
 	}
-	return TEXT("");
+	return TEXT("Unknown");
 }
 
 inline const TCHAR* ToString(const EUpdateConfigCompletionResult Value)
@@ -1870,7 +1870,7 @@ inline const TCHAR* ToString(const EUpdateConfigCompletionResult Value)
 		return TEXT("Succeeded");
 	}
 	}
-	return TEXT("");
+	return TEXT("Unknown");
 }
 
 inline const TCHAR* ToString(const EKickMemberCompletionResult Value)
@@ -1918,7 +1918,7 @@ inline const TCHAR* ToString(const EKickMemberCompletionResult Value)
 		return TEXT("Succeeded");
 	}
 	}
-	return TEXT("");
+	return TEXT("Unknown");
 }
 
 inline const TCHAR* ToString(const EPromoteMemberCompletionResult Value)
@@ -1970,7 +1970,7 @@ inline const TCHAR* ToString(const EPromoteMemberCompletionResult Value)
 		return TEXT("Succeeded");
 	}
 	}
-	return TEXT("");
+	return TEXT("Unknown");
 }
 
 inline const TCHAR* ToString(const PartySystemPermissions::EPresencePermissions Value)
@@ -2018,7 +2018,7 @@ inline const TCHAR* ToString(const PartySystemPermissions::EPresencePermissions
 		return TEXT("AnyonePublishIdAnyonePublishKey");
 	}
 	}
-	return TEXT("");
+	return TEXT("Unknown");
 }
 
 inline const TCHAR* ToString(const PartySystemPermissions::EInvitePermissions Value)
@@ -2038,7 +2038,7 @@ inline const TCHAR* ToString(const PartySystemPermissions::EInvitePermissions Va
 		return TEXT("Anyone");
 	}
 	}
-	return TEXT("");
+	return TEXT("Unknown");
 }
 
 inline const TCHAR* ToString(const EJoinRequestAction Value)
@@ -2058,7 +2058,7 @@ inline const TCHAR* ToString(const EJoinRequestAction Value)
 		return TEXT("AutoReject");
 	}
 	}
-	return TEXT("");
+	return TEXT("Unknown");
 }
 
 inline const TCHAR* ToString(const EInvitationResponse Value)
@@ -2082,7 +2082,7 @@ inline const TCHAR* ToString(const EInvitationResponse Value)
 		return TEXT("Accepted");
 	}
 	}
-	return TEXT("");
+	return TEXT("Unknown");
 }
 
 inline FString ToDebugString(const FPartyConfiguration& PartyConfiguration)
diff --git a/Engine/Plugins/Online/OnlineSubsystem/Source/Public/Interfaces/OnlinePurchaseInterface.h b/Engine/Plugins/Online/OnlineSubsystem/Source/Public/Interfaces/OnlinePurchaseInterface.h
--- a/Engine/Plugins/Online/OnlineSubsystem/Source/Public/Interfaces/OnlinePurchaseInterface.h
+++ b/Engine/Plugins/Online/OnlineSubsystem/Source/Public/Interfaces/OnlinePurchaseInterface.h
@@ -152,6 +152,12 @@ class FRedeemCodeRequest
 public:
 	/** Code to redeem */
 	FString Code;
+
+	/** Optional CodeUseId that was given if code was previously locked before redeeming - See IOnlineCodeRedemption::LockCode */
+	FString CodeUseId;
+
+	/** Where this code is being fulfilled from - e.g. Launcher, GameName*/
+	FString FulfillmentSource;
 };
 
 /**
@@ -176,7 +182,7 @@ DECLARE_MULTICAST_DELEGATE_OneParam(FOnUnexpectedPurchaseReceipt, const FUniqueN
 typedef FOnUnexpectedPurchaseReceipt::FDelegate FOnUnexpectedPurchaseReceiptDelegate;
 
 /**
- *	IOnlinePurchase - Interface for IAP (In App Purchases) services
+ * IOnlinePurchase - Interface for IAP (In App Purchases) services
  */
 class IOnlinePurchase
 {
diff --git a/Engine/Plugins/Online/OnlineSubsystem/Source/Public/OnlineSessionSettings.h b/Engine/Plugins/Online/OnlineSubsystem/Source/Public/OnlineSessionSettings.h
--- a/Engine/Plugins/Online/OnlineSubsystem/Source/Public/OnlineSessionSettings.h
+++ b/Engine/Plugins/Online/OnlineSubsystem/Source/Public/OnlineSessionSettings.h
@@ -49,6 +49,8 @@
 /** Whether there is a skill update in progress or not (don't do multiple at once) */
 /** Whether to shrink the session slots when a player leaves the match or not */
 
+#define INVALID_SESSION_SETTING_ID -1
+
 /**
  *	One setting describing an online session
  * contains a key, value and how this setting is advertised to others, if at all
@@ -67,7 +69,7 @@ public:
 	/** Default constructor, used when serializing a network packet */
 	FOnlineSessionSetting()
 		: AdvertisementType(EOnlineDataAdvertisementType::DontAdvertise)
-		, ID(-1)
+		, ID(INVALID_SESSION_SETTING_ID)
 	{
 	}
 
@@ -76,7 +78,7 @@ public:
 	FOnlineSessionSetting(Type&& InData)
 		: Data(Forward<Type>(InData))
 		, AdvertisementType(EOnlineDataAdvertisementType::DontAdvertise)
-		, ID(-1)
+		, ID(INVALID_SESSION_SETTING_ID)
 	{
 	}
 
@@ -85,7 +87,7 @@ public:
 	FOnlineSessionSetting(Type&& InData, EOnlineDataAdvertisementType::Type InAdvertisementType)
 		: Data(Forward<Type>(InData))
 		, AdvertisementType(InAdvertisementType)
-		, ID(-1)
+		, ID(INVALID_SESSION_SETTING_ID)
 	{
 	}
 
@@ -114,7 +116,14 @@ public:
 
 	FString ToString() const
 	{
-		return FString::Printf(TEXT("%s : %s : %d"), *Data.ToString(), EOnlineDataAdvertisementType::ToString(AdvertisementType), ID);
+		if (ID == INVALID_SESSION_SETTING_ID)
+		{
+			return FString::Printf(TEXT("%s : %s"), *Data.ToString(), EOnlineDataAdvertisementType::ToString(AdvertisementType));
+		}
+		else
+		{
+			return FString::Printf(TEXT("%s : %s : %d"), *Data.ToString(), EOnlineDataAdvertisementType::ToString(AdvertisementType), ID);
+		}
 	}
 };
 
@@ -145,7 +154,7 @@ public:
 	FOnlineSessionSearchParam(Type&& InData)
 		: Data(Forward<Type>(InData))
 		, ComparisonOp(EOnlineComparisonOp::Equals)
-		, ID(-1)
+		, ID(INVALID_SESSION_SETTING_ID)
 	{
 	}
 
@@ -154,7 +163,7 @@ public:
 	FOnlineSessionSearchParam(Type&& InData, EOnlineComparisonOp::Type InComparisonOp)
 		: Data(Forward<Type>(InData))
 		, ComparisonOp(InComparisonOp)
-		, ID(-1)
+		, ID(INVALID_SESSION_SETTING_ID)
 	{
 	}
 
diff --git a/Engine/Plugins/Online/OnlineSubsystem/Source/Public/OnlineSubsystem.h b/Engine/Plugins/Online/OnlineSubsystem/Source/Public/OnlineSubsystem.h
--- a/Engine/Plugins/Online/OnlineSubsystem/Source/Public/OnlineSubsystem.h
+++ b/Engine/Plugins/Online/OnlineSubsystem/Source/Public/OnlineSubsystem.h
@@ -54,6 +54,12 @@ DECLARE_CYCLE_STAT_EXTERN(TEXT("SessionInt"), STAT_Session_Interface, STATGROUP_
 /** Total time to process both local/remote voice */
 DECLARE_CYCLE_STAT_EXTERN(TEXT("VoiceInt"), STAT_Voice_Interface, STATGROUP_Online, ONLINESUBSYSTEM_API);
 
+#if UE_BUILD_SHIPPING
+#define OSS_REDACT(x) TEXT("<Redacted>")
+#else
+#define OSS_REDACT(x) (x)
+#endif
+
 #define ONLINE_LOG_PREFIX TEXT("OSS: ")
 #define UE_LOG_ONLINE(Verbosity, Format, ...) \
 { \
@@ -402,24 +408,24 @@ public:
 	virtual IOnlineChatPtr GetChatInterface() const = 0;
 
 	/**
-	* Get the notification handler instance for this subsystem
-	* @return Pointer for the appropriate notification handler
-	*/
+	 * Get the notification handler instance for this subsystem
+	 * @return Pointer for the appropriate notification handler
+	 */
 	FOnlineNotificationHandlerPtr GetOnlineNotificationHandler() const
 	{
 		return OnlineNotificationHandler;
 	}
 
 	/**
-	* Get the interface for managing turn based multiplayer games
-	* @return Interface pointer for the appropriate online user service
-	*/
+	 * Get the interface for managing turn based multiplayer games
+	 * @return Interface pointer for the appropriate online user service
+	 */
 	virtual IOnlineTurnBasedPtr GetTurnBasedInterface() const = 0;
 
 	/**
-	* Get the transport manager instance for this subsystem
-	* @return Pointer for the appropriate transport manager
-	*/
+	 * Get the transport manager instance for this subsystem
+	 * @return Pointer for the appropriate transport manager
+	 */
 	FOnlineNotificationTransportManagerPtr GetOnlineNotificationTransportManager() const
 	{
 		return OnlineNotificationTransportManager;
diff --git a/Engine/Plugins/Online/OnlineSubsystemFacebook/Source/Private/Windows/OnlineIdentityFacebook.cpp b/Engine/Plugins/Online/OnlineSubsystemFacebook/Source/Private/Windows/OnlineIdentityFacebook.cpp
--- a/Engine/Plugins/Online/OnlineSubsystemFacebook/Source/Private/Windows/OnlineIdentityFacebook.cpp
+++ b/Engine/Plugins/Online/OnlineSubsystemFacebook/Source/Private/Windows/OnlineIdentityFacebook.cpp
@@ -22,6 +22,8 @@ FOnlineIdentityFacebook::FOnlineIdentityFacebook(FOnlineSubsystemFacebook* InSub
 		UE_LOG(LogOnline, Warning, TEXT("Missing bUsePopup= in [OnlineSubsystemFacebook.OnlineIdentityFacebook] of DefaultEngine.ini"));
 	}
 
+	GConfig->GetArray(TEXT("OnlineSubsystemFacebook.OnlineIdentityFacebook"), TEXT("LoginDomains"), LoginDomains, GEngineIni);
+
 	LoginURLDetails.ClientId = InSubsystem->GetAppId();
 	if (LoginURLDetails.ClientId.IsEmpty())
 	{
@@ -258,6 +260,9 @@ bool FOnlineIdentityFacebook::Logout(int32 LocalUserNum)
 		UserIds.Remove(LocalUserNum);
 		// reset scope permissions
 		GConfig->GetArray(TEXT("OnlineSubsystemFacebook.OnlineIdentityFacebook"), TEXT("ScopeFields"), LoginURLDetails.ScopeFields, GEngineIni);
+
+		TriggerOnLoginFlowLogoutDelegates(LoginDomains);
+
 		// not async but should call completion delegate anyway
 		FacebookSubsystem->ExecuteNextTick([this, LocalUserNum, UserId]() {
 			TriggerOnLogoutCompleteDelegates(LocalUserNum, true);
@@ -268,8 +273,7 @@ bool FOnlineIdentityFacebook::Logout(int32 LocalUserNum)
 	}
 	else
 	{
-		UE_LOG(LogOnline, Warning, TEXT("No logged in user found for LocalUserNum=%d."),
-			LocalUserNum);
+		UE_LOG(LogOnline, Warning, TEXT("No logged in user found for LocalUserNum=%d."), LocalUserNum);
 		FacebookSubsystem->ExecuteNextTick([this, LocalUserNum]() {
 			TriggerOnLogoutCompleteDelegates(LocalUserNum, false);
 		});
diff --git a/Engine/Plugins/Online/OnlineSubsystemFacebook/Source/Private/Windows/OnlineIdentityFacebook.h b/Engine/Plugins/Online/OnlineSubsystemFacebook/Source/Private/Windows/OnlineIdentityFacebook.h
--- a/Engine/Plugins/Online/OnlineSubsystemFacebook/Source/Private/Windows/OnlineIdentityFacebook.h
+++ b/Engine/Plugins/Online/OnlineSubsystemFacebook/Source/Private/Windows/OnlineIdentityFacebook.h
@@ -95,6 +95,9 @@ class FOnlineIdentityFacebook :
 	/** Whether we have a registration in flight or not */
 	bool bHasLoginOutstanding;
 
+	/** Domains used for login, for cookie management */
+	TArray<FString> LoginDomains;
+
 public:
 	// IOnlineIdentity
 	
diff --git a/Engine/Plugins/Online/OnlineSubsystemGoogle/Source/Private/IOS/GoogleHelper.cpp b/Engine/Plugins/Online/OnlineSubsystemGoogle/Source/Private/IOS/GoogleHelper.cpp
--- a/Engine/Plugins/Online/OnlineSubsystemGoogle/Source/Private/IOS/GoogleHelper.cpp
+++ b/Engine/Plugins/Online/OnlineSubsystemGoogle/Source/Private/IOS/GoogleHelper.cpp
@@ -7,6 +7,46 @@
 #include "IOSAppDelegate.h"
 #include "IOS/IOSAsyncTask.h"
 
+// Explicit refresh token
+//refreshTokensWithHandler:
+
+// Possibly refresh token, handled by SDK
+//getTokensWithHandler:
+
+bool GetAuthTokenFromGoogleUser(GIDGoogleUser* user, FAuthTokenGoogle& OutAuthToken)
+{
+	bool bSuccess = false;
+
+	OutAuthToken.AccessToken = user.authentication.accessToken;
+	if (!OutAuthToken.AccessToken.IsEmpty())
+	{
+		OutAuthToken.IdToken = user.authentication.idToken;
+		if (!OutAuthToken.IdToken.IsEmpty() && OutAuthToken.IdTokenJWT.Parse(OutAuthToken.IdToken))
+		{
+			OutAuthToken.TokenType = TEXT("Bearer");
+			OutAuthToken.ExpiresIn = 3600.0;
+			OutAuthToken.RefreshToken = user.authentication.refreshToken;
+
+			OutAuthToken.AuthData.Add("refresh_token", OutAuthToken.RefreshToken);
+			OutAuthToken.AuthData.Add("access_token", OutAuthToken.AccessToken);
+			OutAuthToken.AuthData.Add("id_token", OutAuthToken.IdToken);
+			OutAuthToken.AuthType = EGoogleAuthTokenType::AccessToken;
+			OutAuthToken.ExpiresInUTC = FDateTime::UtcNow() + FTimespan(OutAuthToken.ExpiresIn * ETimespan::TicksPerSecond);
+			bSuccess = true;
+		}
+		else
+		{
+			UE_LOG_ONLINE(Verbose, TEXT("GetAuthTokenFromGoogleUser: Failed to parse id token"));
+		}
+	}
+	else
+	{
+		UE_LOG_ONLINE(Verbose, TEXT("GetAuthTokenFromGoogleUser: Access token missing"));
+	}
+
+	return bSuccess;
+}
+
 @implementation FGoogleHelper
 
 - (id) initwithClientId: (NSString*) inClientId withBasicProfile: (bool) bWithProfile
@@ -48,6 +88,12 @@
 	return Delegate.GetHandle();
 }
 
+-(FDelegateHandle)AddOnGoogleDisconnectComplete: (const FOnGoogleSignOutCompleteDelegate&) Delegate
+{
+	_OnDisconnectComplete.Add(Delegate);
+	return Delegate.GetHandle();
+}
+
 - (void) login: (NSArray*) InScopes
 {
 	[self printAuthStatus];
@@ -55,17 +101,34 @@
 	dispatch_async(dispatch_get_main_queue(), ^
 	{
 		GIDSignIn* signIn = [GIDSignIn sharedInstance];
-		[signIn setScopes: InScopes];//[[NSArray alloc] initWithArray: InScopes]];
+		[signIn setScopes: InScopes];
 		[signIn signIn];
 	});
 }
 
 - (void) logout
 {
+	UE_LOG(LogOnline, Display, TEXT("logout"));
+
 	dispatch_async(dispatch_get_main_queue(), ^
 	{
 		GIDSignIn* signIn = [GIDSignIn sharedInstance];
 		[signIn signOut];
+		
+		FGoogleSignOutData SignOutData;
+		SignOutData.Response = EGoogleLoginResponse::RESPONSE_OK;
+
+		[FIOSAsyncTask CreateTaskWithBlock : ^ bool(void)
+		{
+			UE_LOG(LogOnline, Display, TEXT("logoutComplete: %s"), *SignOutData.ToDebugString());
+
+			// Notify on the game thread
+			_OnSignOutComplete.Broadcast(SignOutData);
+			return true;
+		}];
+
+		// Revokes access (use to clear keychain/cache), triggers didDisconnectWithUser
+		// [signIn disconnect];
 	});
 }
 
@@ -74,55 +137,23 @@
 	FGoogleSignInData SignInData;
 	SignInData.ErrorStr = MoveTemp(FString([error localizedDescription]));
 
-	UE_LOG(LogOnline, Display, TEXT("signIn GID:%p User:%p Error:%s"), signIn, user, *SignInData.ErrorStr);
+	UE_LOG(LogOnline, Display, TEXT("signIn didSignInForUser GID:%p User:%p Error:%s"), signIn, user, *SignInData.ErrorStr);
 	[self printAuthStatus];
 
 	if (user)
 	{
-		SignInData.UserId = FString(user.userID);
-		SignInData.IdToken = FString(user.authentication.idToken);
-		SignInData.RealName = FString(user.profile.name);
-		SignInData.FirstName = FString(user.profile.givenName);
-		SignInData.LastName = FString(user.profile.familyName);
-		SignInData.Email = FString(user.profile.email);
-		// user.profile.imageURL
-		SignInData.Response = EGoogleLoginResponse::RESPONSE_OK;
-
-		SignInData.AccessToken = user.authentication.accessToken;
-		SignInData.RefreshToken = user.authentication.refreshToken;
-
-#if __IPHONE_OS_VERSION_MIN_REQUIRED >= __IPHONE_10_0
-		NSISO8601DateFormatter* dateFormatter = [[NSISO8601DateFormatter alloc] init];
-		dateFormatter.formatOptions = NSISO8601DateFormatWithInternetDateTime | NSISO8601DateFormatWithDashSeparatorInDate | NSISO8601DateFormatWithColonSeparatorInTime | NSISO8601DateFormatWithTimeZone;
-#else
-		// see https://developer.apple.com/library/content/qa/qa1480/_index.html
-		NSDateFormatter* dateFormatter = [[NSDateFormatter alloc] init];
-		NSLocale* enUSPOSIXLocale = [NSLocale localeWithLocaleIdentifier:@"en_US_POSIX"];
-		[dateFormatter setLocale:enUSPOSIXLocale];
-		[dateFormatter setDateFormat:@"yyyy-MM-dd'T'HH:mm:ssZZZZZ"];
-#endif
-
-		NSString* accessTokenExpirationDate = [dateFormatter stringFromDate: user.authentication.accessTokenExpirationDate];
-		NSString* idTokenExpirationDate = [dateFormatter stringFromDate: user.authentication.idTokenExpirationDate];
-
-		if (!FDateTime::ParseIso8601(*FString(accessTokenExpirationDate), SignInData.AccessTokenExpiresIn))
+		if (GetAuthTokenFromGoogleUser(user, SignInData.AuthToken))
 		{
-			UE_LOG(LogOnline, Display, TEXT("Failed to parse access token expiration time"));
+			SignInData.Response = EGoogleLoginResponse::RESPONSE_OK;
 		}
-
-		if (!FDateTime::ParseIso8601(*FString(idTokenExpirationDate), SignInData.IdTokenExpiresIn))
+		else
 		{
-			UE_LOG(LogOnline, Display, TEXT("Failed to parse ID token expiration time"));
+			SignInData.Response = EGoogleLoginResponse::RESPONSE_ERROR;
 		}
-
-		[dateFormatter release];
-
-	//refreshTokensWithHandler:
-	//getTokensWithHandler:
 	}
 	else if (error.code == kGIDSignInErrorCodeHasNoAuthInKeychain)
 	{
-		SignInData.Response = EGoogleLoginResponse::RESPONSE_OK;
+		SignInData.Response = EGoogleLoginResponse::RESPONSE_NOAUTH;
 	}
 	else if (error.code == kGIDSignInErrorCodeCanceled)
 	{
@@ -147,17 +178,15 @@
 {
 	FGoogleSignOutData SignOutData;
 	SignOutData.ErrorStr = MoveTemp(FString([error localizedDescription]));
-
-	UE_LOG(LogOnline, Display, TEXT("didDisconnectWithUser %p %p %s"), signIn, user, *SignOutData.ErrorStr);
-	// User disconnected
-
-	UE_LOG(LogOnline, Display, TEXT("SignOut: %s"), *SignOutData.ToDebugString());
+	SignOutData.Response = SignOutData.ErrorStr.IsEmpty() ? EGoogleLoginResponse::RESPONSE_OK : EGoogleLoginResponse::RESPONSE_ERROR;
 
 	[FIOSAsyncTask CreateTaskWithBlock : ^ bool(void)
 	{
-		 // Notify on the game thread
-		 _OnSignOutComplete.Broadcast(SignOutData);
-		 return true;
+		UE_LOG(LogOnline, Display, TEXT("didDisconnectWithUser Complete: %s"), *SignOutData.ToDebugString());
+
+		// Notify on the game thread
+		_OnDisconnectComplete.Broadcast(SignOutData);
+		return true;
 	}];
 }
 
diff --git a/Engine/Plugins/Online/OnlineSubsystemGoogle/Source/Private/IOS/GoogleHelper.h b/Engine/Plugins/Online/OnlineSubsystemGoogle/Source/Private/IOS/GoogleHelper.h
--- a/Engine/Plugins/Online/OnlineSubsystemGoogle/Source/Private/IOS/GoogleHelper.h
+++ b/Engine/Plugins/Online/OnlineSubsystemGoogle/Source/Private/IOS/GoogleHelper.h
@@ -3,6 +3,7 @@
 #pragma once
 
 #include "CoreMinimal.h"
+#include "OnlineSubsystemGoogleTypes.h"
 
 THIRD_PARTY_INCLUDES_START
 #import <GoogleSignIn/GoogleSignIn.h>
@@ -12,10 +13,12 @@ enum class EGoogleLoginResponse : uint8
 {
 	/** Google Sign In SDK ok response */
 	RESPONSE_OK = 0,
+	/** Silent sign in failed */
+	RESPONSE_NOAUTH = 1,
 	/** Google Sign In  SDK user cancellation */
-	RESPONSE_CANCELED = 1,
+	RESPONSE_CANCELED = 2,
 	/** Google Sign In  SDK error */
-	RESPONSE_ERROR = 2,
+	RESPONSE_ERROR = 3
 };
 
 inline const TCHAR* ToString(EGoogleLoginResponse Response)
@@ -24,10 +27,14 @@ inline const TCHAR* ToString(EGoogleLoginResponse Response)
 	{
 		case EGoogleLoginResponse::RESPONSE_OK:
 			return TEXT("RESPONSE_OK");
+		case EGoogleLoginResponse::RESPONSE_NOAUTH:
+			return TEXT("RESPONSE_NOAUTH");
 		case EGoogleLoginResponse::RESPONSE_CANCELED:
 			return TEXT("RESPONSE_CANCELED");
 		case EGoogleLoginResponse::RESPONSE_ERROR:
 			return TEXT("RESPONSE_ERROR");
+		default:
+			break;
 	}
 
 	return TEXT("");
@@ -39,46 +46,18 @@ struct FGoogleSignInData
 	/** @return a string that prints useful debug information about this response */
 	FString ToDebugString() const
 	{
-		if (Response == EGoogleLoginResponse::RESPONSE_OK)
-		{
-			return FString::Printf(TEXT("UserId: %s RealName: %s FirstName: %s LastName: %s Email: %s"),
-								   *UserId,
-								   *RealName,
-								   *FirstName,
-								   *LastName,
-								   *Email);
-		}
-		else
-		{
-			return FString::Printf(TEXT("Response: %s Error: %s"),
-								   ToString(Response),
-								   *ErrorStr);
-		}
+		return FString::Printf(TEXT("Response: %s Valid: %d Error: %s"),
+							   ToString(Response),
+							   AuthToken.IsValid(),
+							   *ErrorStr);
 	}
 
 	/** Result of the sign in */
 	EGoogleLoginResponse Response;
 	/** Error response, if any */
 	FString ErrorStr;
-
-	/** Id associated with the user account provided by the online service during registration */
-	FString UserId;
-	/** Real name */
-	FString RealName;
-	/** First name */
-	FString FirstName;
-	/** Last name */
-	FString LastName;
-	/** JWT */
-	FString IdToken;
-	FDateTime IdTokenExpiresIn;
-	/** Email address */
-	FString Email;
-	/** Access Token */
-	FString AccessToken;
-	FDateTime AccessTokenExpiresIn;
-	/** Refresh Token */
-	FString RefreshToken;
+	/** Token data (access/refresh/id), only valid on RESPONSE_OK */
+	FAuthTokenGoogle AuthToken;
 };
 
 struct FGoogleSignOutData
@@ -117,15 +96,18 @@ typedef FOnGoogleSignOutComplete::FDelegate FOnGoogleSignOutCompleteDelegate;
 {
 	FOnGoogleSignInComplete _OnSignInComplete;
 	FOnGoogleSignOutComplete _OnSignOutComplete;
+	FOnGoogleSignOutComplete _OnDisconnectComplete;
 };
 
 - (id) initwithClientId: (NSString*) inClientId withBasicProfile: (bool) bWithProfile;
 - (void) login: (NSArray*) InScopes;
 - (void) logout;
 
-/** Add a listener to the token change event */
+/** Add a listener to the sign in complete event */
 -(FDelegateHandle)AddOnGoogleSignInComplete: (const FOnGoogleSignInCompleteDelegate&) Delegate;
-/** Add a listener to the user id change event */
+/** Add a listener to the sign out complete event */
 -(FDelegateHandle)AddOnGoogleSignOutComplete: (const FOnGoogleSignOutCompleteDelegate&) Delegate;
+/** Add a listener to the disconnect complete event */
+-(FDelegateHandle)AddOnGoogleDisconnectComplete: (const FOnGoogleSignOutCompleteDelegate&) Delegate;
 
 @end
diff --git a/Engine/Plugins/Online/OnlineSubsystemGoogle/Source/Private/IOS/OnlineIdentityGoogle.cpp b/Engine/Plugins/Online/OnlineSubsystemGoogle/Source/Private/IOS/OnlineIdentityGoogle.cpp
--- a/Engine/Plugins/Online/OnlineSubsystemGoogle/Source/Private/IOS/OnlineIdentityGoogle.cpp
+++ b/Engine/Plugins/Online/OnlineSubsystemGoogle/Source/Private/IOS/OnlineIdentityGoogle.cpp
@@ -38,33 +38,21 @@ bool FOnlineIdentityGoogle::Init()
 
 void FOnlineIdentityGoogle::OnSignInComplete(const FGoogleSignInData& InSignInData)
 {
-	UE_LOG_ONLINE(Verbose, TEXT("OnSignInComplete"));
+	UE_LOG_ONLINE(Verbose, TEXT("OnSignInComplete %s"), ToString(InSignInData.Response));
 
-	FAuthTokenGoogle AuthToken;
-
-	if (InSignInData.Response == EGoogleLoginResponse::RESPONSE_OK)
+	// @todo verify that SignInSilently is working right
+	//if (InSignInData.Response != EGoogleLoginResponse::RESPONSE_NOAUTH)
+	if (LoginCompletionDelegate.IsBound())
 	{
-		AuthToken.AuthType = EGoogleAuthTokenType::AccessToken;
-		AuthToken.AccessToken = InSignInData.AccessToken;
-		AuthToken.TokenType = TEXT("Bearer");
-		AuthToken.RefreshToken = InSignInData.RefreshToken;
-		AuthToken.IdToken = InSignInData.IdToken;
-		if (!AuthToken.IdTokenJWT.Parse(AuthToken.IdToken))
-		{
-			UE_LOG_ONLINE(Verbose, TEXT("Failed to parse id token"));
-		}
-		AuthToken.ExpiresInUTC = InSignInData.AccessTokenExpiresIn;
-		AuthToken.ExpiresIn = (InSignInData.AccessTokenExpiresIn - FDateTime::UtcNow()).GetTotalSeconds();
+		//ensure(LoginCompletionDelegate.IsBound());
+		LoginCompletionDelegate.ExecuteIfBound(InSignInData.Response, InSignInData.AuthToken);
+		LoginCompletionDelegate.Unbind();
 	}
-
-	//ensure(LoginCompletionDelegate.IsBound());
-	LoginCompletionDelegate.ExecuteIfBound(InSignInData.Response, AuthToken);
-	LoginCompletionDelegate.Unbind();
 }
 
 void FOnlineIdentityGoogle::OnSignOutComplete(const FGoogleSignOutData& InSignOutData)
 {
-	UE_LOG_ONLINE(Verbose, TEXT("OnLogoutComplete"));
+	UE_LOG_ONLINE(Verbose, TEXT("OnSignOutComplete %s"), ToString(InSignOutData.Response));
 	ensure(LogoutCompletionDelegate.IsBound());
 	LogoutCompletionDelegate.ExecuteIfBound(InSignOutData.Response);
 	LogoutCompletionDelegate.Unbind();
diff --git a/Engine/Plugins/Online/OnlineSubsystemGoogle/Source/Private/OnlineSubsystemGoogleTypes.h b/Engine/Plugins/Online/OnlineSubsystemGoogle/Source/Private/OnlineSubsystemGoogleTypes.h
--- a/Engine/Plugins/Online/OnlineSubsystemGoogle/Source/Private/OnlineSubsystemGoogleTypes.h
+++ b/Engine/Plugins/Online/OnlineSubsystemGoogle/Source/Private/OnlineSubsystemGoogleTypes.h
@@ -6,6 +6,11 @@
 #include "OnlineJsonSerializer.h"
 #include "OnlineSubsystemGooglePackage.h"
 
+class FAuthTokenGoogle;
+#if PLATFORM_IOS
+@class GIDGoogleUser;
+#endif
+
 
 /**
  * Types of supported auth tokens
@@ -267,6 +272,10 @@ private:
 		ONLINE_JSON_SERIALIZE("refresh_token", RefreshToken);
 		ONLINE_JSON_SERIALIZE("id_token", IdToken);
 	END_ONLINE_JSON_SERIALIZER
+
+#if PLATFORM_IOS
+	friend bool GetAuthTokenFromGoogleUser(GIDGoogleUser* user, FAuthTokenGoogle& OutAuthToken);
+#endif
 };
 
 /**
diff --git a/Engine/Plugins/Online/OnlineSubsystemGoogle/Source/Private/Windows/OnlineIdentityGoogle.cpp b/Engine/Plugins/Online/OnlineSubsystemGoogle/Source/Private/Windows/OnlineIdentityGoogle.cpp
--- a/Engine/Plugins/Online/OnlineSubsystemGoogle/Source/Private/Windows/OnlineIdentityGoogle.cpp
+++ b/Engine/Plugins/Online/OnlineSubsystemGoogle/Source/Private/Windows/OnlineIdentityGoogle.cpp
@@ -18,6 +18,8 @@ FOnlineIdentityGoogle::FOnlineIdentityGoogle(FOnlineSubsystemGoogle* InSubsystem
 		UE_LOG(LogOnline, Warning, TEXT("Missing RedirectPort= in [OnlineSubsystemGoogle.OnlineIdentityGoogle] of DefaultEngine.ini"));
 	}
 
+	GConfig->GetArray(TEXT("OnlineSubsystemGoogle.OnlineIdentityGoogle"), TEXT("LoginDomains"), LoginDomains, GEngineIni);
+
 	LoginURLDetails.ClientId = InSubsystem->GetAppId();
 
 	// Setup permission scope fields
@@ -416,6 +418,9 @@ bool FOnlineIdentityGoogle::Logout(int32 LocalUserNum)
 		UserIds.Remove(LocalUserNum);
 		// reset scope permissions
 		GConfig->GetArray(TEXT("OnlineSubsystemGoogle.OnlineIdentityGoogle"), TEXT("ScopeFields"), LoginURLDetails.ScopeFields, GEngineIni);
+
+		TriggerOnLoginFlowLogoutDelegates(LoginDomains);
+
 		// not async but should call completion delegate anyway
 		GoogleSubsystem->ExecuteNextTick([this, LocalUserNum, UserId]() 
 		{
@@ -427,8 +432,7 @@ bool FOnlineIdentityGoogle::Logout(int32 LocalUserNum)
 	}
 	else
 	{
-		UE_LOG(LogOnline, Warning, TEXT("No logged in user found for LocalUserNum=%d."),
-			LocalUserNum);
+		UE_LOG(LogOnline, Warning, TEXT("No logged in user found for LocalUserNum=%d."), LocalUserNum);
 		GoogleSubsystem->ExecuteNextTick([this, LocalUserNum]() 
 		{
 			TriggerOnLogoutCompleteDelegates(LocalUserNum, false);
diff --git a/Engine/Plugins/Online/OnlineSubsystemGoogle/Source/Private/Windows/OnlineIdentityGoogle.h b/Engine/Plugins/Online/OnlineSubsystemGoogle/Source/Private/Windows/OnlineIdentityGoogle.h
--- a/Engine/Plugins/Online/OnlineSubsystemGoogle/Source/Private/Windows/OnlineIdentityGoogle.h
+++ b/Engine/Plugins/Online/OnlineSubsystemGoogle/Source/Private/Windows/OnlineIdentityGoogle.h
@@ -114,6 +114,9 @@ class FOnlineIdentityGoogle :
 	/** Whether we have a registration in flight or not */
 	bool bHasLoginOutstanding;
 
+	/** Domains used for login, for cookie management */
+	TArray<FString> LoginDomains;
+
 public:
 	// IOnlineIdentity
 	
diff --git a/Engine/Plugins/Online/OnlineSubsystemTwitch/Source/Private/OnlineExternalUITwitch.cpp b/Engine/Plugins/Online/OnlineSubsystemTwitch/Source/Private/OnlineExternalUITwitch.cpp
--- a/Engine/Plugins/Online/OnlineSubsystemTwitch/Source/Private/OnlineExternalUITwitch.cpp
+++ b/Engine/Plugins/Online/OnlineSubsystemTwitch/Source/Private/OnlineExternalUITwitch.cpp
@@ -22,7 +22,7 @@ bool FOnlineExternalUITwitch::ShowLoginUI(const int ControllerIndex, bool bShowO
 			if (URLDetails.IsValid())
 			{
 				const FString& CurrentLoginNonce = IdentityInt->GetCurrentLoginNonce();
-				const FString RequestedURL = URLDetails.GetURL(CurrentLoginNonce);
+				const FString RequestedURL = URLDetails.GetAuthUrl(CurrentLoginNonce);
 				if (OnLoginFlowUIRequiredDelegates.IsBound())
 				{
 					bool bShouldContinueLoginFlow = false;
@@ -33,19 +33,19 @@ bool FOnlineExternalUITwitch::ShowLoginUI(const int ControllerIndex, bool bShowO
 				}
 				else
 				{
-					IOnlineSubsystem* ConsoleSubsystem = IOnlineSubsystem::GetByPlatform();
-					IOnlineExternalUIPtr ConsoleExternalUI = ConsoleSubsystem ? ConsoleSubsystem->GetExternalUIInterface() : nullptr;
-					if (ConsoleExternalUI.IsValid())
+					IOnlineSubsystem* PlatformSubsystem = IOnlineSubsystem::GetByPlatform();
+					IOnlineExternalUIPtr PlatformExternalUI = PlatformSubsystem ? PlatformSubsystem->GetExternalUIInterface() : nullptr;
+					if (PlatformExternalUI.IsValid())
 					{
 						FShowWebUrlParams ShowParams;
 						ShowParams.bEmbedded = false; // TODO:  See if we can pipe along the embedded coordinates
 						ShowParams.bShowBackground = true;
 						ShowParams.bShowCloseButton = true;
 						ShowParams.bResetCookies = true; // potential for previously logged in user
-						ShowParams.CallbackPath = URLDetails.LoginRedirectUrl;
+						ShowParams.CallbackPath = URLDetails.GetLoginRedirectUrl();
 
 						FOnShowWebUrlClosedDelegate OnConsoleShowWebUrlCompleteDelegate = FOnShowWebUrlClosedDelegate::CreateThreadSafeSP(this, &FOnlineExternalUITwitch::OnConsoleShowWebUrlComplete, ControllerIndex, Delegate);
-						bStarted = ConsoleExternalUI->ShowWebURL(RequestedURL, ShowParams, OnConsoleShowWebUrlCompleteDelegate);
+						bStarted = PlatformExternalUI->ShowWebURL(RequestedURL, ShowParams, OnConsoleShowWebUrlCompleteDelegate);
 						if (!bStarted)
 						{
 							UE_LOG_ONLINE(Warning, TEXT("FOnlineExternalUITwitch::ShowLoginUI: Console ShowWebURL failed"));
@@ -85,7 +85,7 @@ FLoginFlowResult FOnlineExternalUITwitch::ParseRedirectResult(const FTwitchLogin
 
 	TMap<FString, FString> ParamsMap;
 
-	FString ResponseStr = RedirectURL.Mid(URLDetails.LoginRedirectUrl.Len() + 1);
+	FString ResponseStr = RedirectURL.Mid(URLDetails.GetLoginRedirectUrl().Len() + 1);
 	{
 		TArray<FString> Params;
 		ResponseStr.ParseIntoArray(Params, TEXT("&"));
@@ -106,7 +106,7 @@ FLoginFlowResult FOnlineExternalUITwitch::ParseRedirectResult(const FTwitchLogin
 		check(IdentityInt.IsValid());
 
 		const FString& CurrentLoginNonce = IdentityInt->GetCurrentLoginNonce();
-		if (CurrentLoginNonce == URLDetails.GetNonce(*State))
+		if (CurrentLoginNonce == URLDetails.ParseNonce(*State))
 		{
 			const FString* AccessToken = ParamsMap.Find(TWITCH_ACCESS_TOKEN);
 			if (AccessToken)
@@ -143,7 +143,7 @@ FLoginFlowResult FOnlineExternalUITwitch::OnLoginRedirectURL(const FString& Redi
 		if (URLDetails.IsValid())
 		{
 			// Wait for the RedirectURI to appear
-			if (RedirectURL.StartsWith(URLDetails.LoginRedirectUrl))
+			if (RedirectURL.StartsWith(URLDetails.GetLoginRedirectUrl()))
 			{
 				Result = ParseRedirectResult(URLDetails, RedirectURL);
 			}
diff --git a/Engine/Plugins/Online/OnlineSubsystemTwitch/Source/Private/OnlineIdentityTwitch.cpp b/Engine/Plugins/Online/OnlineSubsystemTwitch/Source/Private/OnlineIdentityTwitch.cpp
--- a/Engine/Plugins/Online/OnlineSubsystemTwitch/Source/Private/OnlineIdentityTwitch.cpp
+++ b/Engine/Plugins/Online/OnlineSubsystemTwitch/Source/Private/OnlineIdentityTwitch.cpp
@@ -9,6 +9,127 @@
 #include "Http.h"
 #include "Misc/ConfigCacheIni.h"
 
+bool FTwitchLoginURL::IsValid() const
+{
+	return !GetLoginUrl().IsEmpty() && !GetLoginRedirectUrl().IsEmpty() && !Subsystem->GetAppId().IsEmpty();
+}
+
+FString FTwitchLoginURL::GetAuthUrl(const FString& Nonce) const
+{
+	const bool bForceVerify(GetForceVerify());
+	const FString LoginUrl(GetLoginUrl());
+	const FString LoginRedirectUrl(GetLoginRedirectUrl());
+	const FString StatePrefix(GetStatePrefix());
+	const FString ClientId(Subsystem->GetAppId());
+	const TArray<FString> ScopeFields(GetScopeFields());
+
+	const FString Scopes = FString::Join(ScopeFields, TEXT(" "));
+	FString State;
+	if (!StatePrefix.IsEmpty())
+	{
+		State = FString::Printf(TEXT("%s-%s"), *StatePrefix, *Nonce);
+	}
+	else
+	{
+		State = Nonce;
+	}
+
+	return FString::Printf(TEXT("%s?force_verify=%s&response_type=token&client_id=%s&scope=%s&state=%s&redirect_uri=%s"),
+		*LoginUrl,
+		bForceVerify ? TEXT("true") : TEXT("false"),
+		*FGenericPlatformHttp::UrlEncode(ClientId),
+		*FGenericPlatformHttp::UrlEncode(Scopes),
+		*FGenericPlatformHttp::UrlEncode(State),
+		*FGenericPlatformHttp::UrlEncode(LoginRedirectUrl));
+}
+
+FString FTwitchLoginURL::ParseNonce(const FString& State)
+{
+	const FString DecodedState(FGenericPlatformHttp::UrlDecode(State));
+	FString Nonce;
+	// Get everything after the last '-'
+	static const FString HyphenString(TEXT("-"));
+	if (DecodedState.Split(HyphenString, nullptr, &Nonce, ESearchCase::CaseSensitive, ESearchDir::FromEnd) &&
+		!Nonce.IsEmpty())
+	{
+		return Nonce;
+	}
+	// Return the url decoded State parameter
+	return DecodedState;
+}
+
+bool FTwitchLoginURL::GetForceVerify() const
+{
+	bool bForceVerify = false;
+	if (!GConfig->GetBool(TEXT("OnlineSubsystemTwitch.OnlineIdentityTwitch"), TEXT("bForceVerify"), bForceVerify, GEngineIni))
+	{
+		static bool bWarned = false;
+		if (!bWarned)
+		{
+			bWarned = true;
+			UE_LOG_ONLINE(Warning, TEXT("Missing bForceVerify= in [OnlineSubsystemTwitch.OnlineIdentityTwitch] of DefaultEngine.ini"));
+		}
+	}
+	return bForceVerify;
+}
+
+FString FTwitchLoginURL::GetLoginUrl() const
+{
+	FString LoginUrl;
+	if (!GConfig->GetString(TEXT("OnlineSubsystemTwitch.OnlineIdentityTwitch"), TEXT("LoginUrl"), LoginUrl, GEngineIni) ||
+		LoginUrl.IsEmpty())
+	{
+		static bool bWarned = false;
+		if (!bWarned)
+		{
+			bWarned = true;
+			UE_LOG_ONLINE(Warning, TEXT("Missing LoginUrl= in [OnlineSubsystemTwitch.OnlineIdentityTwitch] of DefaultEngine.ini"));
+		}
+	}
+	return LoginUrl;
+}
+
+FString FTwitchLoginURL::GetStatePrefix() const
+{
+	// See if the app has overridden the state prefix
+	if (!StatePrefixOverride.IsEmpty())
+	{
+		return StatePrefixOverride;
+	}
+
+	FString StatePrefix;
+	GConfig->GetString(TEXT("OnlineSubsystemTwitch.OnlineIdentityTwitch"), TEXT("StatePrefix"), StatePrefix, GEngineIni);
+	return StatePrefix;
+}
+
+void FTwitchLoginURL::OverrideStatePrefix(const FString& InStatePrefixOverride)
+{
+	StatePrefixOverride = InStatePrefixOverride;
+}
+
+TArray<FString> FTwitchLoginURL::GetScopeFields() const
+{
+	TArray<FString> ScopeFields;
+	GConfig->GetArray(TEXT("OnlineSubsystemTwitch.OnlineIdentityTwitch"), TEXT("ScopeFields"), ScopeFields, GEngineIni);
+	return ScopeFields;
+}
+
+FString FTwitchLoginURL::GetLoginRedirectUrl() const
+{
+	FString LoginRedirectUrl;
+	if (!GConfig->GetString(TEXT("OnlineSubsystemTwitch.OnlineIdentityTwitch"), TEXT("LoginRedirectUrl"), LoginRedirectUrl, GEngineIni) ||
+		LoginRedirectUrl.IsEmpty())
+	{
+		static bool bWarned = false;
+		if (!bWarned)
+		{
+			bWarned = true;
+			UE_LOG_ONLINE(Warning, TEXT("Missing LoginRedirectUrl= in [OnlineSubsystemTwitch.OnlineIdentityTwitch] of DefaultEngine.ini"));
+		}
+	}
+	return LoginRedirectUrl;
+}
+
 TSharedPtr<FUserOnlineAccountTwitch> FOnlineIdentityTwitch::GetUserAccountTwitch(const FUniqueNetId& UserId) const
 {
 	TSharedPtr<FUserOnlineAccountTwitch> Result;
@@ -78,7 +199,7 @@ bool FOnlineIdentityTwitch::Login(int32 LocalUserNum, const FOnlineAccountCreden
 	else if (!LoginURLDetails.IsValid())
 	{
 		ErrorStr = FString::Printf(TEXT("OnlineSubsystemTwitch is improperly configured in DefaultEngine.ini LoginURL=%s LoginRedirectUrl=%s ClientId=%s"),
-			*LoginURLDetails.LoginUrl, *LoginURLDetails.LoginRedirectUrl, *LoginURLDetails.ClientId);
+			*LoginURLDetails.GetLoginUrl(), *LoginURLDetails.GetLoginRedirectUrl(), *Subsystem->GetAppId());
 	}
 	else if (LocalUserNum < 0 || LocalUserNum >= MAX_LOCAL_PLAYERS)
 	{
@@ -142,7 +263,8 @@ void FOnlineIdentityTwitch::OnValidateAuthTokenComplete(int32 LocalUserNum, cons
 			// Confirm we have all of the scope permissions we require
 			bool bHasAllPermissions = true;
 			FString MissingScopeFields;
-			for (const FString& ScopeField : LoginURLDetails.ScopeFields)
+			const TArray<FString> RequiredScopeFields(LoginURLDetails.GetScopeFields());
+			for (const FString& ScopeField : RequiredScopeFields)
 			{
 				const bool bHasPermission = User->GetScopePermissions().Contains(ScopeField);
 				if (!bHasPermission)
@@ -190,7 +312,18 @@ void FOnlineIdentityTwitch::ValidateAuthToken(int32 LocalUserNum, const FOnlineA
 	// kick off http request to validate access token
 	TSharedRef<IHttpRequest> HttpRequest = FHttpModule::Get().CreateRequest();
 
-	FString Url(FString::Printf(TEXT("%s?client_id=%s"), *LoginURLDetails.TokenValidateUrl, *FGenericPlatformHttp::UrlEncode(LoginURLDetails.ClientId)));
+	FString TokenValidateUrl;
+	if (!GConfig->GetString(TEXT("OnlineSubsystemTwitch.OnlineIdentityTwitch"), TEXT("TokenValidateUrl"), TokenValidateUrl, GEngineIni) ||
+		TokenValidateUrl.IsEmpty())
+	{
+		static bool bWarned = false;
+		if (!bWarned)
+		{
+			bWarned = true;
+			UE_LOG_ONLINE(Warning, TEXT("Missing TokenValidateUrl= in [OnlineSubsystemTwitch.OnlineIdentityTwitch] of DefaultEngine.ini"));
+		}
+	}
+	FString Url(FString::Printf(TEXT("%s?client_id=%s"), *TokenValidateUrl, *FGenericPlatformHttp::UrlEncode(Subsystem->GetAppId())));
 
 	HttpRequest->OnProcessRequestComplete().BindThreadSafeSP(this, &FOnlineIdentityTwitch::ValidateAuthToken_HttpRequestComplete, LocalUserNum, AccountCredentials, InCompletionDelegate);
 	HttpRequest->SetURL(MoveTemp(Url));
@@ -373,6 +506,9 @@ void FOnlineIdentityTwitch::OnTwitchLogoutComplete(const FUniqueNetId& UserId)
 		// remove cached user id
 		UserIds.Remove(LocalUserNum);
 
+		TArray<FString> LoginDomains;
+		GConfig->GetArray(TEXT("OnlineSubsystemTwitch.OnlineIdentityTwitch"), TEXT("LoginDomains"), LoginDomains, GEngineIni);
+
 		TriggerOnLoginFlowLogoutDelegates(LoginDomains);
 
 		TSharedRef<const FUniqueNetId> UserIdRef(UserId.AsShared());
@@ -419,12 +555,22 @@ void FOnlineIdentityTwitch::RevokeAuthTokenInternal(const FUniqueNetId& UserId,
 	// kick off http request to validate access token
 	TSharedRef<IHttpRequest> HttpRequest = FHttpModule::Get().CreateRequest();
 
-	FString Url(TokenRevokeUrl);
-	FString PostData(FString::Printf(TEXT("client_id=%s&token=%s"), *FGenericPlatformHttp::UrlEncode(LoginURLDetails.ClientId), *FGenericPlatformHttp::UrlEncode(AuthToken)));
+	FString TokenRevokeUrl;
+	if (!GConfig->GetString(TEXT("OnlineSubsystemTwitch.OnlineIdentityTwitch"), TEXT("TokenRevokeUrl"), TokenRevokeUrl, GEngineIni) ||
+		TokenRevokeUrl.IsEmpty())
+	{
+		static bool bWarned = false;
+		if (!bWarned)
+		{
+			bWarned = true;
+			UE_LOG_ONLINE(Warning, TEXT("Missing TokenRevokeUrl= in [OnlineSubsystemTwitch.OnlineIdentityTwitch] of DefaultEngine.ini"));
+		}
+	}
+	FString PostData(FString::Printf(TEXT("client_id=%s&token=%s"), *FGenericPlatformHttp::UrlEncode(Subsystem->GetAppId()), *FGenericPlatformHttp::UrlEncode(AuthToken)));
 
 	TSharedRef<const FUniqueNetId> UserIdRef(UserId.AsShared());
 	HttpRequest->OnProcessRequestComplete().BindThreadSafeSP(this, &FOnlineIdentityTwitch::RevokeAuthToken_HttpRequestComplete, UserIdRef, InCompletionDelegate);
-	HttpRequest->SetURL(TokenRevokeUrl);
+	HttpRequest->SetURL(MoveTemp(TokenRevokeUrl));
 	HttpRequest->SetHeader(TEXT("Accept"), Subsystem->GetTwitchApiVersion());
 	HttpRequest->SetHeader(TEXT("Content-Type"), TEXT("application/x-www-form-urlencoded"));
 	HttpRequest->SetVerb(TEXT("POST"));
@@ -557,40 +703,11 @@ FString FOnlineIdentityTwitch::GetAuthToken(int32 LocalUserNum) const
 
 FOnlineIdentityTwitch::FOnlineIdentityTwitch(FOnlineSubsystemTwitch* InSubsystem)
 	: Subsystem(InSubsystem)
+	, LoginURLDetails(InSubsystem)
 	, bHasLoginOutstanding(false)
 	, ZeroId(MakeShared<FUniqueNetIdString>())
 {
 	check(InSubsystem);
-	if (!GConfig->GetString(TEXT("OnlineSubsystemTwitch.OnlineIdentityTwitch"), TEXT("LoginUrl"), LoginURLDetails.LoginUrl, GEngineIni))
-	{
-		UE_LOG_ONLINE(Warning, TEXT("Missing LoginUrl= in [OnlineSubsystemTwitch.OnlineIdentityTwitch] of DefaultEngine.ini"));
-	}
-	if (!GConfig->GetBool(TEXT("OnlineSubsystemTwitch.OnlineIdentityTwitch"), TEXT("bForceVerify"), LoginURLDetails.bForceVerify, GEngineIni))
-	{
-		UE_LOG_ONLINE(Warning, TEXT("Missing bForceVerify= in [OnlineSubsystemTwitch.OnlineIdentityTwitch] of DefaultEngine.ini"));
-	}
-	if (!GConfig->GetString(TEXT("OnlineSubsystemTwitch.OnlineIdentityTwitch"), TEXT("LoginRedirectUrl"), LoginURLDetails.LoginRedirectUrl, GEngineIni))
-	{
-		UE_LOG_ONLINE(Warning, TEXT("Missing LoginRedirectUrl= in [OnlineSubsystemTwitch.OnlineIdentityTwitch] of DefaultEngine.ini"));
-	}
-	if (!GConfig->GetString(TEXT("OnlineSubsystemTwitch.OnlineIdentityTwitch"), TEXT("TokenValidateUrl"), LoginURLDetails.TokenValidateUrl, GEngineIni))
-	{
-		UE_LOG_ONLINE(Warning, TEXT("Missing TokenValidateUrl= in [OnlineSubsystemTwitch.OnlineIdentityTwitch] of DefaultEngine.ini"));
-	}
-	// StatePrefix is optional
-	GConfig->GetString(TEXT("OnlineSubsystemTwitch.OnlineIdentityTwitch"), TEXT("StatePrefix"), LoginURLDetails.StatePrefix, GEngineIni);
-
-	LoginURLDetails.ClientId = Subsystem->GetAppId();
-
-	GConfig->GetArray(TEXT("OnlineSubsystemTwitch.OnlineIdentityTwitch"), TEXT("LoginDomains"), LoginDomains, GEngineIni);
-
-	// Setup permission scope fields
-	GConfig->GetArray(TEXT("OnlineSubsystemTwitch.OnlineIdentityTwitch"), TEXT("ScopeFields"), LoginURLDetails.ScopeFields, GEngineIni);
-
-	if (!GConfig->GetString(TEXT("OnlineSubsystemTwitch.OnlineIdentityTwitch"), TEXT("TokenRevokeUrl"), TokenRevokeUrl, GEngineIni))
-	{
-		UE_LOG_ONLINE(Warning, TEXT("Missing TokenRevokeUrl= in [OnlineSubsystemTwitch.OnlineIdentityTwitch] of DefaultEngine.ini"));
-	}
 }
 
 FPlatformUserId FOnlineIdentityTwitch::GetPlatformUserIdFromUniqueNetId(const FUniqueNetId& UniqueNetId)
@@ -616,3 +733,9 @@ FString FOnlineIdentityTwitch::GetAuthType() const
 {
 	return TEXT("twitch");
 }
+
+void FOnlineIdentityTwitch::SetStatePrefix(const FString& StatePrefix)
+{
+	UE_LOG_ONLINE(Log, TEXT("FOnlineIdentityTwitch::SetStatePrefix: Setting StatePrefix to %s"), *StatePrefix);
+	LoginURLDetails.OverrideStatePrefix(StatePrefix);
+}
diff --git a/Engine/Plugins/Online/OnlineSubsystemTwitch/Source/Private/OnlineSubsystemTwitch.cpp b/Engine/Plugins/Online/OnlineSubsystemTwitch/Source/Private/OnlineSubsystemTwitch.cpp
--- a/Engine/Plugins/Online/OnlineSubsystemTwitch/Source/Private/OnlineSubsystemTwitch.cpp
+++ b/Engine/Plugins/Online/OnlineSubsystemTwitch/Source/Private/OnlineSubsystemTwitch.cpp
@@ -130,11 +130,6 @@ bool FOnlineSubsystemTwitch::Init()
 {
 	UE_LOG_ONLINE(Verbose, TEXT("FOnlineSubsystemTwitch::Init() Name: %s"), *InstanceName.ToString());
 
-	if (!GConfig->GetString(TEXT("OnlineSubsystemTwitch"), TEXT("ClientId"), ClientId, GEngineIni))
-	{
-		UE_LOG(LogOnline, Warning, TEXT("Missing ClientId= in [OnlineSubsystemTwitch] of DefaultEngine.ini"));
-	}
-
 	TwitchIdentity = MakeShared<FOnlineIdentityTwitch, ESPMode::ThreadSafe>(this);
 	TwitchExternalUIInterface = MakeShared<FOnlineExternalUITwitch, ESPMode::ThreadSafe>(this);
 
@@ -168,6 +163,17 @@ bool FOnlineSubsystemTwitch::Shutdown()
 
 FString FOnlineSubsystemTwitch::GetAppId() const
 {
+	FString ClientId;
+	if (!GConfig->GetString(TEXT("OnlineSubsystemTwitch"), TEXT("ClientId"), ClientId, GEngineIni) ||
+		ClientId.IsEmpty())
+	{
+		static bool bWarned = false;
+		if (!bWarned)
+		{
+			bWarned = true;
+			UE_LOG(LogOnline, Warning, TEXT("Missing ClientId= in [OnlineSubsystemTwitch] of DefaultEngine.ini"));
+		}
+	}
 	return ClientId;
 }
 
diff --git a/Engine/Plugins/Online/OnlineSubsystemTwitch/Source/Public/OnlineIdentityTwitch.h b/Engine/Plugins/Online/OnlineSubsystemTwitch/Source/Public/OnlineIdentityTwitch.h
--- a/Engine/Plugins/Online/OnlineSubsystemTwitch/Source/Public/OnlineIdentityTwitch.h
+++ b/Engine/Plugins/Online/OnlineSubsystemTwitch/Source/Public/OnlineIdentityTwitch.h
@@ -28,35 +28,22 @@ typedef TMap<FString, TSharedRef<FUserOnlineAccountTwitch> > FUserOnlineAccountT
 /**
  * Contains URL details for Twitch interaction
  */
-struct FTwitchLoginURL
+struct ONLINESUBSYSTEMTWITCH_API FTwitchLoginURL
 {
-	FTwitchLoginURL()
-		: bForceVerify(false)
+	FTwitchLoginURL(FOnlineSubsystemTwitch* InSubsystem)
+		: Subsystem(InSubsystem)
 	{
 	}
 
-	/** The endpoint at Twitch we are supposed to hit for auth */
-	FString LoginUrl;
-	/** Whether or not we should always prompt the user for authorization */
-	bool bForceVerify;
-	/** The redirect url for Twitch to redirect to upon completion. */
-	FString LoginRedirectUrl;
-	/** The url to validate a OAuth token with Twitch */
-	FString TokenValidateUrl;
-	/** The client id given to us by Twitch */
-	FString ClientId;
-	/** Config based list of permission scopes to use when logging in */
-	TArray<FString> ScopeFields;
-	/** Prefix for the state field */
-	FString StatePrefix;
-
 	/** 
 	 * @return whether this is properly configured or not
 	 */
-	bool IsValid() const
-	{
-		return !LoginUrl.IsEmpty() && !LoginRedirectUrl.IsEmpty() && !ClientId.IsEmpty();
-	}
+	bool IsValid() const;
+
+	/** 
+	 * @return the auth url to spawn in the browser
+	 */
+	FString GetAuthUrl(const FString& Nonce) const;
 
 	/** 
 	 * @return a random nonce
@@ -67,46 +54,49 @@ struct FTwitchLoginURL
 		return FGuid::NewGuid().ToString();
 	}
 
+	/**
+	 * Get the Nonce from the specified State
+	 */
+	static FString ParseNonce(const FString& State);
+
+	// Settings accessors
 	/** 
-	 * @return the auth url to spawn in the browser
+	 * @return the login url
 	 */
-	FString GetURL(const FString& Nonce) const
-	{
-		FString Scopes = FString::Join(ScopeFields, TEXT(" "));
-		FString State;
-		if (!StatePrefix.IsEmpty())
-		{
-			State = FString::Printf(TEXT("%s+%s"), *StatePrefix, *Nonce);
-		}
-		else
-		{
-			State = Nonce;
-		}
-
-		return FString::Printf(TEXT("%s?force_verify=%s&response_type=token&client_id=%s&scope=%s&state=%s&redirect_uri=%s"),
-			*LoginUrl,
-			bForceVerify ? TEXT("true") : TEXT("false"),
-			*FGenericPlatformHttp::UrlEncode(ClientId),
-			*FGenericPlatformHttp::UrlEncode(Scopes),
-			*FGenericPlatformHttp::UrlEncode(State),
-			*FGenericPlatformHttp::UrlEncode(LoginRedirectUrl));
-	}
+	FString GetLoginUrl() const;
 
+	/**
+	 * @return the ForceVerify setting
+	 */
+	bool GetForceVerify() const;
+	
 	/** 
-	 * Get the Nonce from the specified State
+	 * @return the redirect url
 	 */
-	FString GetNonce(const FString& State) const
-	{
-		FString DecodedState(FGenericPlatformHttp::UrlDecode(State));
-		if (StatePrefix.IsEmpty() ||
-			DecodedState.Len() < (StatePrefix.Len() + 1) ||
-			!DecodedState.StartsWith(StatePrefix) ||
-			DecodedState[StatePrefix.Len()] != TEXT('+'))
-		{
-			return DecodedState;
-		}
-		return DecodedState.Mid(StatePrefix.Len() + 1);
-	}
+	FString GetLoginRedirectUrl() const;
+
+	/** 
+	 * @return the state prefix to put before the nonce
+	 */
+	FString GetStatePrefix() const;
+
+	/** 
+	 * Override the state prefix
+	 * @param InStatePrefixOverride the state prefix override
+	 */
+	void OverrideStatePrefix(const FString& InStatePrefixOverride);
+
+	/**
+	 * @return the scope fields we are requesting permissions for
+	 */
+	TArray<FString> GetScopeFields() const;
+
+private:
+	/** Owning subsystem */
+	FOnlineSubsystemTwitch* Subsystem;
+
+	/** Overridden StatePrefix */
+	FString StatePrefixOverride;
 };
 
 /**
@@ -138,9 +128,14 @@ public:
 	virtual FPlatformUserId GetPlatformUserIdFromUniqueNetId(const FUniqueNetId& UniqueNetId) override;
 	virtual FString GetAuthType() const override;
 
-PACKAGE_SCOPE:
-
 	// FOnlineIdentityTwitch
+	/** 
+	 * Set the state prefix
+	 * @param StatePrefix the state prefix
+	 */
+	void SetStatePrefix(const FString& StatePrefix);
+
+PACKAGE_SCOPE:
 
 	/**
 	 * Constructor
@@ -258,14 +253,10 @@ private:
 
 	/** Const details about communicating with twitch API */
 	FTwitchLoginURL LoginURLDetails;
-	/** Domains used for login, for cookie management */
-	TArray<FString> LoginDomains;
 	/** Nonce for current login attempt */
 	FString CurrentLoginNonce;
 	/** Whether we have a registration in flight or not */
 	bool bHasLoginOutstanding;
-	/** URL used to revoke an access token */
-	FString TokenRevokeUrl;
 
 	/** Re-usable empty unique id for errors */
 	TSharedRef<FUniqueNetId> ZeroId;
diff --git a/Engine/Plugins/Online/OnlineSubsystemTwitch/Source/Public/OnlineSubsystemTwitch.h b/Engine/Plugins/Online/OnlineSubsystemTwitch/Source/Public/OnlineSubsystemTwitch.h
--- a/Engine/Plugins/Online/OnlineSubsystemTwitch/Source/Public/OnlineSubsystemTwitch.h
+++ b/Engine/Plugins/Online/OnlineSubsystemTwitch/Source/Public/OnlineSubsystemTwitch.h
@@ -101,9 +101,6 @@ private:
 	/** Interface for external UI services on Twitch */
 	FOnlineExternalUITwitchPtr TwitchExternalUIInterface;
 
-	/** The client id given to us by Twitch */
-	FString ClientId;
-
 	/** Twitch API version */
 	FString TwitchApiVersion;
 };
diff --git a/Engine/Plugins/Online/OnlineSubsystemUtils/Source/OnlineSubsystemUtils/Classes/IpNetDriver.h b/Engine/Plugins/Online/OnlineSubsystemUtils/Source/OnlineSubsystemUtils/Classes/IpNetDriver.h
--- a/Engine/Plugins/Online/OnlineSubsystemUtils/Source/OnlineSubsystemUtils/Classes/IpNetDriver.h
+++ b/Engine/Plugins/Online/OnlineSubsystemUtils/Source/OnlineSubsystemUtils/Classes/IpNetDriver.h
@@ -85,4 +85,21 @@ class ONLINESUBSYSTEMUTILS_API UIpNetDriver : public UNetDriver
 	// It may get called multiple times in a single frame if additional processing after a previous alert exceeds the threshold again
 	DECLARE_MULTICAST_DELEGATE(FOnNetworkProcessingCausingSlowFrame);
 	static FOnNetworkProcessingCausingSlowFrame OnNetworkProcessingCausingSlowFrame;
+
+private:
+	/** Number of bytes that will be passed to FSocket::SetReceiveBufferSize when initializing a server. */
+	UPROPERTY(Config)
+	uint32 ServerDesiredSocketReceiveBufferBytes;
+
+	/** Number of bytes that will be passed to FSocket::SetSendBufferSize when initializing a server. */
+	UPROPERTY(Config)
+	uint32 ServerDesiredSocketSendBufferBytes;
+
+	/** Number of bytes that will be passed to FSocket::SetReceiveBufferSize when initializing a client. */
+	UPROPERTY(Config)
+	uint32 ClientDesiredSocketReceiveBufferBytes;
+
+	/** Number of bytes that will be passed to FSocket::SetSendBufferSize when initializing a client. */
+	UPROPERTY(Config)
+	uint32 ClientDesiredSocketSendBufferBytes;
 };
diff --git a/Engine/Plugins/Online/OnlineSubsystemUtils/Source/OnlineSubsystemUtils/Private/IpConnection.cpp b/Engine/Plugins/Online/OnlineSubsystemUtils/Source/OnlineSubsystemUtils/Private/IpConnection.cpp
--- a/Engine/Plugins/Online/OnlineSubsystemUtils/Source/OnlineSubsystemUtils/Private/IpConnection.cpp
+++ b/Engine/Plugins/Online/OnlineSubsystemUtils/Source/OnlineSubsystemUtils/Private/IpConnection.cpp
@@ -25,6 +25,8 @@ Notes:
 #define IP_HEADER_SIZE     (20)
 #define UDP_HEADER_SIZE    (IP_HEADER_SIZE+8)
 
+DECLARE_CYCLE_STAT(TEXT("IpConnection InitRemoteConnection"), Stat_IpConnectionInitRemoteConnection, STATGROUP_Net);
+
 UIpConnection::UIpConnection(const FObjectInitializer& ObjectInitializer) :
 	Super(ObjectInitializer),
 	RemoteAddr(NULL),
@@ -77,6 +79,8 @@ void UIpConnection::InitLocalConnection(UNetDriver* InDriver, class FSocket* InS
 
 void UIpConnection::InitRemoteConnection(UNetDriver* InDriver, class FSocket* InSocket, const FURL& InURL, const class FInternetAddr& InRemoteAddr, EConnectionState InState, int32 InMaxPacket, int32 InPacketOverhead)
 {
+	SCOPE_CYCLE_COUNTER(Stat_IpConnectionInitRemoteConnection);
+
 	InitBase(InDriver, InSocket, InURL, InState, 
 		// Use the default packet size/overhead unless overridden by a child class
 		(InMaxPacket == 0 || InMaxPacket > MAX_PACKET_SIZE) ? MAX_PACKET_SIZE : InMaxPacket,
diff --git a/Engine/Plugins/Online/OnlineSubsystemUtils/Source/OnlineSubsystemUtils/Private/IpNetDriver.cpp b/Engine/Plugins/Online/OnlineSubsystemUtils/Source/OnlineSubsystemUtils/Private/IpNetDriver.cpp
--- a/Engine/Plugins/Online/OnlineSubsystemUtils/Source/OnlineSubsystemUtils/Private/IpNetDriver.cpp
+++ b/Engine/Plugins/Online/OnlineSubsystemUtils/Source/OnlineSubsystemUtils/Private/IpNetDriver.cpp
@@ -13,13 +13,14 @@ Notes:
 #include "Engine/World.h"
 #include "Engine/Engine.h"
 #include "UObject/Package.h"
-#include "PacketAudit.h"
 #include "PacketHandlers/StatelessConnectHandlerComponent.h"
 #include "Engine/NetConnection.h"
 #include "Engine/ChildConnection.h"
 #include "SocketSubsystem.h"
 #include "IpConnection.h"
 
+#include "PacketAudit.h"
+
 #include "IPAddress.h"
 #include "Sockets.h"
 
@@ -32,6 +33,8 @@ Notes:
 	Declarations.
 -----------------------------------------------------------------------------*/
 
+DECLARE_CYCLE_STAT(TEXT("IpNetDriver Add new connection"), Stat_IpNetDriverAddNewConnection, STATGROUP_Net);
+
 UIpNetDriver::FOnNetworkProcessingCausingSlowFrame UIpNetDriver::OnNetworkProcessingCausingSlowFrame;
 
 // Time before the alarm delegate is called (in seconds)
@@ -55,6 +58,10 @@ FAutoConsoleVariableRef GIpNetDriverLongFramePrintoutThresholdSecsCVar(
 
 UIpNetDriver::UIpNetDriver(const FObjectInitializer& ObjectInitializer)
 	: Super(ObjectInitializer)
+	, ServerDesiredSocketReceiveBufferBytes(0x20000)
+	, ServerDesiredSocketSendBufferBytes(0x20000)
+	, ClientDesiredSocketReceiveBufferBytes(0x8000)
+	, ClientDesiredSocketSendBufferBytes(0x8000)
 {
 }
 
@@ -131,9 +138,9 @@ bool UIpNetDriver::InitBase( bool bInitAsClient, FNetworkNotify* InNotify, const
 	}
 
 	// Increase socket queue size, because we are polling rather than threading
-	// and thus we rely on Windows Sockets to buffer a lot of data on the server.
-	int32 RecvSize = bInitAsClient ? 0x8000 : 0x20000;
-	int32 SendSize = bInitAsClient ? 0x8000 : 0x20000;
+	// and thus we rely on the OS socket to buffer a lot of data.
+	int32 RecvSize = bInitAsClient ? ClientDesiredSocketReceiveBufferBytes	: ServerDesiredSocketReceiveBufferBytes;
+	int32 SendSize = bInitAsClient ? ClientDesiredSocketSendBufferBytes		: ServerDesiredSocketSendBufferBytes;
 	Socket->SetReceiveBufferSize(RecvSize,RecvSize);
 	Socket->SetSendBufferSize(SendSize,SendSize);
 	UE_LOG(LogInit, Log, TEXT("%s: Socket queue %i / %i"), SocketSubsystem->GetSocketAPIName(), RecvSize, SendSize );
@@ -416,6 +423,8 @@ void UIpNetDriver::TickDispatch( float DeltaTime )
 
 					if (bPassedChallenge)
 					{
+						SCOPE_CYCLE_COUNTER(Stat_IpNetDriverAddNewConnection);
+
 						UE_LOG(LogNet, Log, TEXT("Server accepting post-challenge connection from: %s"), *FromAddr->ToString(true));
 
 						Connection = NewObject<UIpConnection>(GetTransientPackage(), NetConnectionClass);
@@ -506,12 +515,10 @@ void UIpNetDriver::LowLevelSend(FString Address, void* Data, int32 CountBits)
 		int32 BytesSent = 0;
 		uint32 CountBytes = FMath::DivideAndRoundUp(CountBits, 8);
 
-		FPacketAudit::NotifyLowLevelSend((uint8*)Data, CountBytes, CountBits);
-
 		if (CountBits > 0)
 		{
 			CLOCK_CYCLES(SendCycles);
-			Socket->SendTo(DataToSend, CountBytes, BytesSent, *RemoteAddr);
+			Socket->SendTo(DataToSend, FMath::DivideAndRoundUp(CountBits, 8), BytesSent, *RemoteAddr);
 			UNCLOCK_CYCLES(SendCycles);
 		}
 
diff --git a/Engine/Plugins/Online/OnlineSubsystemUtils/Source/OnlineSubsystemUtils/Private/OnlineBeaconClient.cpp b/Engine/Plugins/Online/OnlineSubsystemUtils/Source/OnlineSubsystemUtils/Private/OnlineBeaconClient.cpp
--- a/Engine/Plugins/Online/OnlineSubsystemUtils/Source/OnlineSubsystemUtils/Private/OnlineBeaconClient.cpp
+++ b/Engine/Plugins/Online/OnlineSubsystemUtils/Source/OnlineSubsystemUtils/Private/OnlineBeaconClient.cpp
@@ -14,16 +14,6 @@
 
 #define BEACON_RPC_TIMEOUT 15.0f
 
-/** For backwards compatibility with the engine packet handler code */
-#ifndef PACKETHANDLER_HAS_BEGINHANDSHAKING
-	#define PACKETHANDLER_HAS_BEGINHANDSHAKING 0
-#endif
-
-/** For backwards compatibility with engine encryption support */
-#ifndef SUPPORTS_ENCRYPTION_TOKEN
-	#define SUPPORTS_ENCRYPTION_TOKEN 0
-#endif
-
 AOnlineBeaconClient::AOnlineBeaconClient(const FObjectInitializer& ObjectInitializer) :
 	Super(ObjectInitializer),
 	BeaconOwner(nullptr),
@@ -111,14 +101,10 @@ bool AOnlineBeaconClient::InitClient(FURL& URL)
 
 				if (BeaconConnection->Handler.IsValid())
 				{
-#if PACKETHANDLER_HAS_BEGINHANDSHAKING
 					BeaconConnection->Handler->BeginHandshaking(
 						FPacketHandlerHandshakeComplete::CreateUObject(this, &AOnlineBeaconClient::SendInitialJoin));
 
 					bSentHandshake = true;
-#else
-					BeaconConnection->StatelessConnectComponent.Pin()->SendInitialConnect();
-#endif
 				}
 
 				SetConnectionState(EBeaconConnectionState::Pending);
@@ -164,16 +150,13 @@ void AOnlineBeaconClient::SendInitialJoin()
 
 		uint32 LocalNetworkVersion = FNetworkVersion::GetLocalNetworkVersion();
 
-#if SUPPORTS_ENCRYPTION_TOKEN
 		if (CVarNetAllowEncryption.GetValueOnGameThread() == 0)
 		{
 			EncryptionToken.Reset();
 		}
 
 		FNetControlMessage<NMT_Hello>::Send(NetDriver->ServerConnection, IsLittleEndian, LocalNetworkVersion, EncryptionToken);
-#else
-		FNetControlMessage<NMT_Hello>::Send(NetDriver->ServerConnection, IsLittleEndian, LocalNetworkVersion);
-#endif
+
 		NetDriver->ServerConnection->FlushNet();
 	}
 }
@@ -253,7 +236,6 @@ void AOnlineBeaconClient::NotifyControlMessage(UNetConnection* Connection, uint8
 #endif
 		switch (MessageType)
 		{
-#if SUPPORTS_ENCRYPTION_TOKEN
 		case NMT_EncryptionAck:
 			{
 				if (FNetDelegates::OnReceivedNetworkEncryptionAck.IsBound())
@@ -269,7 +251,6 @@ void AOnlineBeaconClient::NotifyControlMessage(UNetConnection* Connection, uint8
 				}
 				break;
 			}
-#endif
 		case NMT_BeaconWelcome:
 			{
 				Connection->ClientResponse = TEXT("0");
diff --git a/Engine/Plugins/Online/OnlineSubsystemUtils/Source/OnlineSubsystemUtils/Private/OnlineBeaconHost.cpp b/Engine/Plugins/Online/OnlineSubsystemUtils/Source/OnlineSubsystemUtils/Private/OnlineBeaconHost.cpp
--- a/Engine/Plugins/Online/OnlineSubsystemUtils/Source/OnlineSubsystemUtils/Private/OnlineBeaconHost.cpp
+++ b/Engine/Plugins/Online/OnlineSubsystemUtils/Source/OnlineSubsystemUtils/Private/OnlineBeaconHost.cpp
@@ -11,11 +11,6 @@
 #include "Net/DataChannel.h"
 #include "OnlineBeaconClient.h"
 
-/** For backwards compatibility with engine encryption support */
-#ifndef SUPPORTS_ENCRYPTION_TOKEN
-	#define SUPPORTS_ENCRYPTION_TOKEN 0
-#endif
-
 AOnlineBeaconHost::AOnlineBeaconHost(const FObjectInitializer& ObjectInitializer) :
 	Super(ObjectInitializer)
 {
@@ -99,11 +94,8 @@ void AOnlineBeaconHost::NotifyControlMessage(UNetConnection* Connection, uint8 M
 				uint32 LocalNetworkVersion = FNetworkVersion::GetLocalNetworkVersion();
 				FString EncryptionToken;
 
-#if SUPPORTS_ENCRYPTION_TOKEN
 				FNetControlMessage<NMT_Hello>::Receive(Bunch, IsLittleEndian, RemoteNetworkVersion, EncryptionToken);
-#else
-				FNetControlMessage<NMT_Hello>::Receive(Bunch, IsLittleEndian, RemoteNetworkVersion);
-#endif
+
 				if (!FNetworkVersion::IsNetworkCompatible(LocalNetworkVersion, RemoteNetworkVersion))
 				{
 					UE_LOG(LogBeacon, Error, TEXT("Client not network compatible %s (Local=%d, Remote=%d)"), *Connection->GetName(), LocalNetworkVersion, RemoteNetworkVersion);
@@ -116,7 +108,6 @@ void AOnlineBeaconHost::NotifyControlMessage(UNetConnection* Connection, uint8 M
 					{
 						SendWelcomeControlMessage(Connection);
 					}
-#if SUPPORTS_ENCRYPTION_TOKEN
 					else
 					{
 						if (FNetDelegates::OnReceivedNetworkEncryptionToken.IsBound())
@@ -132,7 +123,6 @@ void AOnlineBeaconHost::NotifyControlMessage(UNetConnection* Connection, uint8 M
 							bCloseConnection = true;
 						}
 					}
-#endif
 				}
 				break;
 			}
diff --git a/Engine/Plugins/Online/OnlineSubsystemUtils/Source/OnlineSubsystemUtils/Private/OnlineEngineInterfaceImpl.cpp b/Engine/Plugins/Online/OnlineSubsystemUtils/Source/OnlineSubsystemUtils/Private/OnlineEngineInterfaceImpl.cpp
--- a/Engine/Plugins/Online/OnlineSubsystemUtils/Source/OnlineSubsystemUtils/Private/OnlineEngineInterfaceImpl.cpp
+++ b/Engine/Plugins/Online/OnlineSubsystemUtils/Source/OnlineSubsystemUtils/Private/OnlineEngineInterfaceImpl.cpp
@@ -513,11 +513,11 @@ void UOnlineEngineInterfaceImpl::BindToExternalUIOpening(const FOnlineExternalUI
 		}
 	}
 
-	IOnlineSubsystem* SubSystemConsole = IOnlineSubsystem::GetByPlatform();
-	if (SubSystemConsole != nullptr &&
-		SubSystem != SubSystemConsole)
+	IOnlineSubsystem* PlatformSubSystem = IOnlineSubsystem::GetByPlatform();
+	if (PlatformSubSystem != nullptr &&
+		SubSystem != PlatformSubSystem)
 	{
-		IOnlineExternalUIPtr ExternalUI = SubSystemConsole->GetExternalUIInterface();
+		IOnlineExternalUIPtr ExternalUI = PlatformSubSystem->GetExternalUIInterface();
 		if (ExternalUI.IsValid())
 		{
 			FOnExternalUIChangeDelegate OnExternalUIChangeDelegate;
diff --git a/Engine/Source/Runtime/Engine/Public/Net/DataChannel.h b/Engine/Source/Runtime/Engine/Public/Net/DataChannel.h
--- a/Engine/Source/Runtime/Engine/Public/Net/DataChannel.h
+++ b/Engine/Source/Runtime/Engine/Public/Net/DataChannel.h
@@ -13,9 +13,6 @@
 #include "Net/DataBunch.h"
 #include "Engine/ChildConnection.h"
 
-// Backwards-compatibility macro for online game plugins to test
-#define SUPPORTS_ENCRYPTION_TOKEN 1
-
 /*-----------------------------------------------------------------------------
 	UControlChannel base class.
 -----------------------------------------------------------------------------*/
diff --git a/Engine/Source/Runtime/PacketHandlers/PacketHandler/Public/PacketHandler.h b/Engine/Source/Runtime/PacketHandlers/PacketHandler/Public/PacketHandler.h
--- a/Engine/Source/Runtime/PacketHandlers/PacketHandler/Public/PacketHandler.h
+++ b/Engine/Source/Runtime/PacketHandlers/PacketHandler/Public/PacketHandler.h
@@ -8,9 +8,6 @@
 #include "Modules/ModuleInterface.h"
 #include "Containers/Queue.h"
 
-/** Temporary version-check macro for BeginHandshaking function, for backwards-compatibility. */
-#define PACKETHANDLER_HAS_BEGINHANDSHAKING 1
-
 PACKETHANDLER_API DECLARE_LOG_CATEGORY_EXTERN(PacketHandlerLog, Log, All);
 
 
