Set up private PCH files everywhere that previously used them.

#rnx

diff --git a/Engine/Plugins/Online/OnlineSubsystemTwitch/Source/OnlineSubsystemTwitch.Build.cs b/Engine/Plugins/Online/OnlineSubsystemTwitch/Source/OnlineSubsystemTwitch.Build.cs
--- a/Engine/Plugins/Online/OnlineSubsystemTwitch/Source/OnlineSubsystemTwitch.Build.cs
+++ b/Engine/Plugins/Online/OnlineSubsystemTwitch/Source/OnlineSubsystemTwitch.Build.cs
@@ -7,6 +7,8 @@ public class OnlineSubsystemTwitch : ModuleRules
 {
 	public OnlineSubsystemTwitch(ReadOnlyTargetRules Target) : base(Target)
 	{
+		PrivatePCHHeaderFile = "Private/OnlineSubsystemTwitchPrivate.h";
+
 		PrivateDefinitions.Add("ONLINESUBSYSTEM_TWITCH_PACKAGE=1");
 		PCHUsage = ModuleRules.PCHUsageMode.UseSharedPCHs;
 
