Cleaned up OnlineSubsystemTwitch to not set a private pch file(it isn't force included) and to enable IWYU.

diff --git a/Engine/Plugins/Online/OnlineSubsystemTwitch/Source/OnlineSubsystemTwitch.Build.cs b/Engine/Plugins/Online/OnlineSubsystemTwitch/Source/OnlineSubsystemTwitch.Build.cs
--- a/Engine/Plugins/Online/OnlineSubsystemTwitch/Source/OnlineSubsystemTwitch.Build.cs
+++ b/Engine/Plugins/Online/OnlineSubsystemTwitch/Source/OnlineSubsystemTwitch.Build.cs
@@ -7,10 +7,7 @@ public class OnlineSubsystemTwitch : ModuleRules
 {
 	public OnlineSubsystemTwitch(ReadOnlyTargetRules Target) : base(Target)
 	{
-		PrivatePCHHeaderFile = "Private/OnlineSubsystemTwitchPrivate.h";
-
 		PrivateDefinitions.Add("ONLINESUBSYSTEM_TWITCH_PACKAGE=1");
-		PCHUsage = ModuleRules.PCHUsageMode.UseSharedPCHs;
 
 		PrivateIncludePaths.Add("Private");
 
diff --git a/Engine/Plugins/Online/OnlineSubsystemTwitch/Source/Private/OnlineAccountTwitch.cpp b/Engine/Plugins/Online/OnlineSubsystemTwitch/Source/Private/OnlineAccountTwitch.cpp
--- a/Engine/Plugins/Online/OnlineSubsystemTwitch/Source/Private/OnlineAccountTwitch.cpp
+++ b/Engine/Plugins/Online/OnlineSubsystemTwitch/Source/Private/OnlineAccountTwitch.cpp
@@ -1,9 +1,16 @@
 // Copyright Epic Games, Inc. All Rights Reserved.
 
-#include "OnlineSubsystemTwitchPrivate.h"
 #include "OnlineAccountTwitch.h"
+#include "OnlineSubsystemTwitchPrivate.h"
 #include "TwitchTokenValidationResponse.h"
 
+FUserOnlineAccountTwitch::FUserOnlineAccountTwitch(const FString& InUserId /*= FString()*/, const FString& InAuthTicket /*= FString()*/) 
+	: UserId(FUniqueNetIdTwitch::Create(InUserId))
+	, AuthTicket(InAuthTicket)
+{
+
+}
+
 FUniqueNetIdRef FUserOnlineAccountTwitch::GetUserId() const
 {
 	return UserId;
diff --git a/Engine/Plugins/Online/OnlineSubsystemTwitch/Source/Private/OnlineSubsystemTwitch.cpp b/Engine/Plugins/Online/OnlineSubsystemTwitch/Source/Private/OnlineSubsystemTwitch.cpp
--- a/Engine/Plugins/Online/OnlineSubsystemTwitch/Source/Private/OnlineSubsystemTwitch.cpp
+++ b/Engine/Plugins/Online/OnlineSubsystemTwitch/Source/Private/OnlineSubsystemTwitch.cpp
@@ -1,8 +1,9 @@
 // Copyright Epic Games, Inc. All Rights Reserved.
 
-#include "OnlineSubsystemTwitchPrivate.h"
-#include "OnlineIdentityTwitch.h"
+#include "OnlineSubsystemTwitch.h"
 #include "OnlineExternalUITwitch.h"
+#include "OnlineIdentityTwitch.h"
+#include "OnlineSubsystemTwitchPrivate.h"
 #include "Engine/Console.h"
 #include "Misc/ConfigCacheIni.h"
 
diff --git a/Engine/Plugins/Online/OnlineSubsystemTwitch/Source/Public/OnlineAccountTwitch.h b/Engine/Plugins/Online/OnlineSubsystemTwitch/Source/Public/OnlineAccountTwitch.h
--- a/Engine/Plugins/Online/OnlineSubsystemTwitch/Source/Public/OnlineAccountTwitch.h
+++ b/Engine/Plugins/Online/OnlineSubsystemTwitch/Source/Public/OnlineAccountTwitch.h
@@ -30,12 +30,7 @@ public:
 	virtual bool GetAuthAttribute(const FString& AttrName, FString& OutAttrValue) const override;
 
 	// FUserOnlineAccountTwitch
-	FUserOnlineAccountTwitch(const FString& InUserId = FString(), const FString& InAuthTicket = FString())
-		: UserId(FUniqueNetIdTwitch::Create(InUserId))
-		, AuthTicket(InAuthTicket)
-	{
-	}
-
+	FUserOnlineAccountTwitch(const FString& InUserId = FString(), const FString& InAuthTicket = FString());
 	virtual ~FUserOnlineAccountTwitch() = default;
 
 	/**
